Scribble Plugin First Steps

#rnx

diff --git a/Engine/Plugins/Experimental/Animation/Scribble/Scribble.uplugin b/Engine/Plugins/Experimental/Animation/Scribble/Scribble.uplugin
new file mode 100644
--- /dev/null
+++ b/Engine/Plugins/Experimental/Animation/Scribble/Scribble.uplugin
@@ -0,0 +1,32 @@
+{
+	"FileVersion" : 3,
+	"Version" : 1,
+	"VersionName" : "1.0",
+	"FriendlyName" : "Scribble",
+	"Description" : "A user interface plugin providing scribble capabilities.",
+	"Category" : "Animation",
+	"CreatedBy" : "Epic Games, Inc.",
+	"CreatedByURL" : "https://epicgames.com",
+	"DocsURL" : "",
+	"MarketplaceURL" : "",
+	"SupportURL" : "",
+	"EnabledByDefault" : false,
+	"CanContainContent" : true,
+	"Installed" : false,
+	"Plugins":
+	[
+	],
+	"Modules" :
+	[
+		{
+			"Name" : "Scribble",
+			"Type" : "Runtime",
+			"LoadingPhase" : "PreDefault"
+		},
+		{
+			"Name" : "ScribbleEditor",
+			"Type" : "Editor",
+			"LoadingPhase" : "PostEngineInit"
+		}
+	]
+}
\ No newline at end of file
diff --git a/Engine/Plugins/Experimental/Animation/Scribble/Source/Scribble/Private/ScribbleGraph.cpp b/Engine/Plugins/Experimental/Animation/Scribble/Source/Scribble/Private/ScribbleGraph.cpp
new file mode 100644
--- /dev/null
+++ b/Engine/Plugins/Experimental/Animation/Scribble/Source/Scribble/Private/ScribbleGraph.cpp
@@ -0,0 +1,554 @@
+// Copyright Epic Games, Inc. All Rights Reserved.
+
+#include "ScribbleGraph.h"
+#include "ScribbleObjectVersion.h"
+
+#include UE_INLINE_GENERATED_CPP_BY_NAME(ScribbleGraph)
+
+#define LOCTEXT_NAMESPACE "ScribbleGraph"
+
+TAutoConsoleVariable<bool> CVarScribbleEnabled(TEXT("Scribble.Enabled"), false, TEXT("enables the scribbling capabilities on supported widgets / editors"));
+
+void FScribbleGraphData::Reset()
+{
+	ViewOffset = FVector2f::ZeroVector;
+	ZoomAmount = 1.0f;
+	Nodes.Reset();
+	GuidToNode.Reset();
+	NotifyChanged();
+}
+
+bool FScribbleGraphData::IsEmpty() const
+{
+	return Nodes.IsEmpty();
+}
+
+int32 FScribbleGraphData::NumNodes() const
+{
+	return Nodes.Num();
+}
+
+FScribbleNode* FScribbleGraphData::GetNode(int32 InNodeIndex)
+{
+	if (Nodes.IsValidIndex(InNodeIndex))
+	{
+		return Nodes[InNodeIndex].Get();
+	}
+	return nullptr;
+}
+
+const FScribbleNode* FScribbleGraphData::GetNode(int32 InNodeIndex) const
+{
+	if (Nodes.IsValidIndex(InNodeIndex))
+	{
+		return Nodes[InNodeIndex].Get();
+	}
+	return nullptr;
+}
+
+TSharedPtr<FScribbleNode> FScribbleGraphData::GetNodePtr(int32 InNodeIndex) const
+{
+	if (Nodes.IsValidIndex(InNodeIndex))
+	{
+		return Nodes[InNodeIndex];
+	}
+	return nullptr;
+}
+
+FScribbleNode* FScribbleGraphData::FindNode(const FGuid& InId)
+{
+	if (const TSharedPtr<FScribbleNode>* Ptr = GuidToNode.Find(InId))
+	{
+		return Ptr->Get();
+	}
+	return nullptr;
+}
+
+const FScribbleNode* FScribbleGraphData::FindNode(const FGuid& InId) const
+{
+	if (const TSharedPtr<FScribbleNode>* Ptr = GuidToNode.Find(InId))
+	{
+		return Ptr->Get();
+	}
+	return nullptr;
+}
+
+TSharedPtr<FScribbleNode> FScribbleGraphData::FindNodePtr(const FGuid& InId) const
+{
+	if (const TSharedPtr<FScribbleNode>* Ptr = GuidToNode.Find(InId))
+	{
+		return *Ptr;
+	}
+	return nullptr;
+}
+
+FScribbleNode* FScribbleGraphData::operator[](int32 InNodeIndex)
+{
+	return GetNode(InNodeIndex);
+}
+
+const FScribbleNode* FScribbleGraphData::operator[](int32 InNodeIndex) const
+{
+	return GetNode(InNodeIndex);
+}
+
+FGuid FScribbleGraphData::AddNode(const TSharedPtr<FScribbleNode>& InNode)
+{
+	if (!InNode)
+	{
+		return FGuid();
+	}
+	
+	if(InNode->GetGraph() == this)
+	{
+		return InNode->GetId();
+	}
+
+	if (const FScribbleNode* ExistingNode = FindNode(InNode->GetId()))
+	{
+		return ExistingNode->GetId();
+	}
+
+	check(!InNode->WeakGraph.IsValid());
+
+	if (!InNode->Guid.IsValid())
+	{
+		InNode->Guid = FGuid::NewGuid();
+	}
+
+	InNode->WeakGraph = AsWeak();
+
+	Modify();
+	Nodes.Add(InNode);
+	GuidToNode.Add(InNode->Guid, InNode);
+
+	NotifyChanged();
+	return InNode->Guid;
+}
+
+bool FScribbleGraphData::RemoveNode(const TSharedPtr<FScribbleNode>& InNode)
+{
+	if (!InNode.IsValid())
+	{
+		return false;
+	}
+	
+	Modify();
+
+	const int32 NumRemoved = Nodes.Remove(InNode);
+	if (NumRemoved == 0)
+	{
+		return false;
+	}
+
+	InNode->WeakGraph.Reset();
+	
+	GuidToNode.Remove(InNode->Guid);
+
+	if (NumRemoved > 0)
+	{
+		NotifyChanged();
+	}
+	return NumRemoved > 0;
+}
+
+bool FScribbleGraphData::RemoveNode(const FGuid& InId)
+{
+	TSharedPtr<FScribbleNode> Node = FindNodePtr(InId);
+	if (!Node)
+	{
+		return false;
+	}
+	return RemoveNode(Node);
+}
+
+bool FScribbleGraphData::SupportsAnchors() const
+{
+	return GetCurrentAnchorDelegate.IsBound() && ResolveAnchorDelegate.IsBound();
+}
+
+FName FScribbleGraphData::GetCurrentAnchor() const
+{
+	if (GetCurrentAnchorDelegate.IsBound())
+	{
+		return GetCurrentAnchorDelegate.Execute();
+	}
+	return NAME_None;
+}
+
+TSharedPtr<FScribbleNode> FScribbleGraphData::GroupNodes(const TArray<TSharedPtr<FScribbleNode>>& InNodes)
+{
+	int32 NumLineStrips = 0;
+	TArray<TSharedPtr<FScribbleNode>> NodesToGroup;
+	for (const TSharedPtr<FScribbleNode>& Node : InNodes)
+	{
+		if (!Node)
+		{
+			continue;
+		}
+
+		if (Node->GetGraph() != this)
+		{
+			continue;
+		}
+
+		// for now we only support lines strips for grouping
+		if (Node->GetType() != EScribbleNodeType::LineStrip)
+		{
+			continue;
+		}
+
+		NodesToGroup.Add(Node);
+
+		if (const FLineStripScribbleNode* LineStripScribbleNode = Cast<FLineStripScribbleNode>(Node.Get()))
+		{
+			NumLineStrips += LineStripScribbleNode->LineStrips.Num();
+		}
+	}
+	
+	if (NodesToGroup.Num() < 2)
+	{
+		return nullptr;
+	}
+
+	Modify();
+	IncrementChangeBracket();
+
+	TSharedPtr<FLineStripScribbleNode> GroupedNode = MakeShared<FLineStripScribbleNode>();
+	GroupedNode->LineStrips.Reserve(NumLineStrips);
+	AddNode(GroupedNode);
+
+	const FVector2f GroupedNodePosition = GroupedNode->GetPosition();
+	for (const TSharedPtr<FScribbleNode>& Node : NodesToGroup)
+	{
+		const FLineStripScribbleNode* LineStripScribbleNode = Cast<FLineStripScribbleNode>(Node.Get());
+		if (!LineStripScribbleNode)
+		{
+			continue;
+		}
+
+		const FVector2f UngroupedNodePosition = LineStripScribbleNode->GetPosition();
+		for (const FScribbleLineStrip& OldLineStrip : LineStripScribbleNode->LineStrips)
+		{
+			GroupedNode->LineStrips.Add(OldLineStrip);
+			FScribbleLineStrip& NewLineStrip = GroupedNode->LineStrips.Last();
+			for (FVector2f& Vertex : NewLineStrip.Vertices)
+			{
+				Vertex = (Vertex + UngroupedNodePosition) - GroupedNodePosition;
+			}
+		}
+	}
+
+	const FBox2f Bounds = GroupedNode->GetContentBounds();
+	if (Bounds.bIsValid)
+	{
+		const FVector2f TopLeft = GroupedNode->GetPosition() + Bounds.GetCenter() - Bounds.GetExtent();
+		GroupedNode->OffsetPosition(TopLeft);
+	}
+
+	DecrementChangeBracket();
+	return GroupedNode;
+}
+
+TArray<TSharedPtr<FScribbleNode>> FScribbleGraphData::UngroupNode(const TSharedPtr<FScribbleNode>& InNode)
+{
+	TArray<TSharedPtr<FScribbleNode>> UngroupedNodes;
+	if (!InNode || InNode->GetGraph() != this)
+	{
+		return UngroupedNodes;
+	}
+
+	// for now we only allow ungrouping of line strips
+	if (InNode->GetType() != EScribbleNodeType::LineStrip)
+	{
+		return UngroupedNodes;
+	}
+
+	const FLineStripScribbleNode* LineStripScribbleNode = Cast<FLineStripScribbleNode>(InNode.Get());
+	if (!LineStripScribbleNode)
+	{
+		return UngroupedNodes;
+	}
+
+	Modify();
+	IncrementChangeBracket();
+
+	const FVector2f GroupedNodePosition = LineStripScribbleNode->GetPosition();
+	for (const FScribbleLineStrip& LineStrip : LineStripScribbleNode->LineStrips)
+	{
+		TSharedPtr<FLineStripScribbleNode> UngroupedNode = MakeShared<FLineStripScribbleNode>();
+		UngroupedNode->LineStrips.Add(LineStrip);
+		AddNode(UngroupedNode);
+		const FVector2f UngroupedNodePosition = UngroupedNode->GetPosition();
+
+		FScribbleLineStrip& NewLineStrip = UngroupedNode->LineStrips.Last();
+		for (FVector2f& Vertex : NewLineStrip.Vertices)
+		{
+			Vertex = (Vertex + GroupedNodePosition) - UngroupedNodePosition;
+		}
+
+		const FBox2f Bounds = UngroupedNode->GetContentBounds();
+		if (Bounds.bIsValid)
+		{
+			const FVector2f TopLeft = UngroupedNode->GetPosition() + Bounds.GetCenter() - Bounds.GetExtent();
+			UngroupedNode->OffsetPosition(TopLeft);
+			UngroupedNodes.Add(UngroupedNode);
+		}
+	}
+
+	DecrementChangeBracket();
+	return UngroupedNodes;
+}
+
+bool FScribbleGraphData::Serialize(FArchive& Ar)
+{
+	Ar.UsingCustomVersion(FScribbleObjectVersion::GUID);
+	IncrementChangeBracket();
+
+	Ar << ViewOffset;
+	Ar << ZoomAmount;
+	
+	int32 Num = Nodes.Num();
+	Ar << Num;
+
+	if (Ar.IsLoading())
+	{
+		Nodes.Reset();
+		GuidToNode.Reset();
+
+		for (int32 Index = 0; Index < Num; ++Index)
+		{
+			uint8 EnumAsByte = 0;
+			Ar << EnumAsByte;
+
+			TSharedPtr<FScribbleNode> Node;
+			switch (static_cast<EScribbleNodeType::Type>(EnumAsByte))
+			{
+				case EScribbleNodeType::LineStrip:
+				{
+					Node = MakeShared<FLineStripScribbleNode>();
+					break;
+				}
+				default:
+				{
+					checkNoEntry();
+					break;
+				}
+			}
+
+			if (!Node)
+			{
+				continue;
+			}
+
+			if (!Node->Serialize(Ar))
+			{
+				DecrementChangeBracket();
+				return false;
+			}
+			
+			Nodes.Add(Node);
+			GuidToNode.Add(Node->GetId(), Node);
+			Node->WeakGraph = AsWeak();
+			NotifyChanged();	
+		}
+	}
+	else
+	{
+		for (int32 Index = 0; Index < Num; ++Index)
+		{
+			uint8 EnumAsByte = Nodes[Index]->GetType();
+			Ar << EnumAsByte;
+			if (!Nodes[Index]->Serialize(Ar))
+			{
+				DecrementChangeBracket();
+				return false;
+			}
+		}
+	}
+	
+	DecrementChangeBracket();
+	return true;
+}
+
+void FScribbleGraphData::IncrementChangeBracket()
+{
+	ChangeBracket++;
+}
+
+void FScribbleGraphData::DecrementChangeBracket()
+{
+	ChangeBracket = FMath::Max(ChangeBracket - 1, 0);
+	NotifyChanged();
+}
+
+void FScribbleGraphData::Modify()
+{
+	OnModifyEvent.Broadcast();
+}
+
+void FScribbleGraphData::SetView(const FVector2f& InViewOffset, float InZoomAmount)
+{
+	if (FMath::IsNearlyEqual(InZoomAmount, ZoomAmount) &&
+		InViewOffset.Equals(ViewOffset))
+	{
+		return;
+	}
+
+	Modify();
+	ViewOffset = InViewOffset;
+	ZoomAmount = InZoomAmount;
+}
+
+void FScribbleGraphData::NotifyChanged()
+{
+	if (ChangeBracket == 0 && AccumulatedChanges > 0)
+	{
+		AccumulatedChanges = 0;
+		OnChangedEvent.Broadcast();
+	}
+	else
+	{
+		AccumulatedChanges++;
+	}
+}
+
+TOptional<FVector2f> FScribbleGraphData::ResolveAnchor(const FName& InName) const
+{
+	if (ResolveAnchorDelegate.IsBound())
+	{
+		return ResolveAnchorDelegate.Execute(InName);
+	}
+	return TOptional<FVector2f>();
+}
+
+void FScribbleGraph::Reset()
+{
+	if (!Data.IsValid())
+	{
+		Data = MakeShared<FScribbleGraphData>();
+	}
+	Data->Reset();
+}
+
+bool FScribbleGraph::IsValid() const
+{
+	return Data.IsValid();
+}
+
+FScribbleGraphData* FScribbleGraph::GetData()
+{
+	if (Data.IsValid())
+	{
+		return Data.Get();
+	}
+	return nullptr;
+}
+
+const FScribbleGraphData* FScribbleGraph::GetData() const
+{
+	if (Data.IsValid())
+	{
+		return Data.Get();
+	}
+	return nullptr;
+}
+
+TSharedPtr<FScribbleGraphData> FScribbleGraph::GetOrCreateDataPtr()
+{
+	if (!Data.IsValid())
+	{
+		Data = MakeShared<FScribbleGraphData>();
+	}
+	return Data;
+}
+
+bool FScribbleGraph::Serialize(FArchive& Ar)
+{
+	Ar.UsingCustomVersion(FScribbleObjectVersion::GUID);
+
+	if (Ar.IsLoading())
+	{
+		if (!Data.IsValid())
+		{
+			Data = MakeShared<FScribbleGraphData>();
+		}
+		
+		bool bContainsData = false;
+		Ar << bContainsData;
+		if (bContainsData)
+		{
+			if (!Data->Serialize(Ar))
+			{
+				return false;
+			}
+		}
+		else
+		{
+			Data->Reset();
+		}
+	}
+	else if (Ar.IsSaving())
+	{
+		bool bContainsData = Data.IsValid() && CVarScribbleEnabled.GetValueOnAnyThread();
+		Ar << bContainsData;
+		if (bContainsData)
+		{
+			if (!Data->Serialize(Ar))
+			{
+				return false;
+			}
+		}
+	}
+	
+	return true;
+}
+
+bool FScribbleGraph::Identical(const FScribbleGraph* Other, uint32 PortFlags) const
+{
+	if (!Other)
+	{
+		return false;
+	}
+	if (Data.IsValid() != Other->Data.IsValid())
+	{
+		return false;
+	}
+	if (!Data.IsValid())
+	{
+		return true;
+	}
+
+	const TSharedPtr<FScribbleGraphData>& OtherData = Other->Data;
+
+	if (Data->NumNodes() != OtherData->NumNodes())
+	{
+		return false;
+	}
+	if (!FMath::IsNearlyEqual(Data->GetZoomAmount(), OtherData->GetZoomAmount()))
+	{
+		return false;
+	}
+	if (!(Data->GetViewOffset() - OtherData->GetViewOffset()).IsNearlyZero())
+	{
+		return false;
+	}
+
+	for (int32 NodeIndex = 0; NodeIndex < Data->NumNodes(); NodeIndex++)
+	{
+		const FScribbleNode* Node = Data->GetNode(NodeIndex);
+		const FScribbleNode* OtherNode = OtherData->GetNode(NodeIndex);
+		if (!Node || !OtherNode)
+		{
+			return false;
+		}
+		if(!Node->Identical(OtherNode, PortFlags))
+		{
+			return false;
+		}
+	}
+
+	return true;
+}
+
+#undef LOCTEXT_NAMESPACE
diff --git a/Engine/Plugins/Experimental/Animation/Scribble/Source/Scribble/Private/ScribbleModule.cpp b/Engine/Plugins/Experimental/Animation/Scribble/Source/Scribble/Private/ScribbleModule.cpp
new file mode 100644
--- /dev/null
+++ b/Engine/Plugins/Experimental/Animation/Scribble/Source/Scribble/Private/ScribbleModule.cpp
@@ -0,0 +1,26 @@
+// Copyright Epic Games, Inc. All Rights Reserved.
+
+#include "Modules/ModuleManager.h"
+#include "ScribbleObjectVersion.h"
+#include "UObject/DevObjectVersion.h"
+
+// Unique Scribble Object version id
+const FGuid FScribbleObjectVersion::GUID(0xDC49359B, 0x53C04DE7, 0x9126EA88, 0x5E1C5D39);
+
+// Register Scribble custom version with Core
+static FDevVersionRegistration GRegisterScribbleObjectVersion(FScribbleObjectVersion::GUID, FScribbleObjectVersion::LatestVersion, TEXT("Dev-Scribble"));
+
+class FScribbleModule : public IModuleInterface
+{
+public:
+	/** IModuleInterface implementation */
+	virtual void StartupModule() override
+	{
+	}
+	
+	virtual void ShutdownModule() override
+	{
+	}
+};
+
+IMPLEMENT_MODULE(FScribbleModule, Scribble)
diff --git a/Engine/Plugins/Experimental/Animation/Scribble/Source/Scribble/Private/ScribbleNode.cpp b/Engine/Plugins/Experimental/Animation/Scribble/Source/Scribble/Private/ScribbleNode.cpp
new file mode 100644
--- /dev/null
+++ b/Engine/Plugins/Experimental/Animation/Scribble/Source/Scribble/Private/ScribbleNode.cpp
@@ -0,0 +1,526 @@
+// Copyright Epic Games, Inc. All Rights Reserved.
+
+#include "ScribbleNode.h"
+#include "ScribbleSettings.h"
+#include "ScribbleGraph.h"
+#include "ScribbleObjectVersion.h"
+#if WITH_EDITOR
+#include "Framework/Application/SlateApplication.h"
+#include "Styling/AppStyle.h"
+#endif
+
+#include UE_INLINE_GENERATED_CPP_BY_NAME(ScribbleNode)
+
+#define LOCTEXT_NAMESPACE "ScribbleNode"
+
+void FScribbleNode::SetAnchor(const FName& InAnchor)
+{
+	const FVector2f OldPosition = GetPosition();
+	Anchor = InAnchor;
+	UpdateRelativePosition(OldPosition);
+}
+
+FVector2f FScribbleNode::GetPosition(bool bResolveAnchor) const
+{
+	if (bResolveAnchor && !Anchor.IsNone())
+	{
+		if (const FScribbleGraphData* Graph = GetGraph())
+		{
+			if (const TOptional<FVector2f> AnchorPosition = Graph->ResolveAnchor(Anchor))
+			{
+				return AnchorPosition.GetValue() + RelativePosition;
+			}
+		}
+	}
+	return Position;
+}
+
+void FScribbleNode::SetPosition(const FVector2f& InPosition)
+{
+	const FVector2f CurrentPosition = GetPosition(false);
+	if (CurrentPosition.Equals(InPosition, 0.001f))
+	{
+		return;
+	}
+
+	Modify();
+	Position = InPosition;
+	UpdateRelativePosition(InPosition);
+	PositionChangedEvent.Broadcast(Position);
+	NotifyGraphChanged();
+}
+
+void FScribbleNode::UpdateRelativePosition(const FVector2f& InAbsolutePosition)
+{
+	RelativePosition = InAbsolutePosition;
+	if (!Anchor.IsNone())
+	{
+		if (const FScribbleGraphData* Graph = GetGraph())
+		{
+			if (const TOptional<FVector2f> AnchorPosition = Graph->ResolveAnchor(Anchor))
+			{
+				RelativePosition = InAbsolutePosition - AnchorPosition.GetValue();
+			}
+		}
+	}
+}
+
+void FScribbleNode::SetSize(const FVector2f& InSize)
+{
+	const FVector2f NewSize(FMath::Max<float>(1,InSize.X), FMath::Max<float>(1,InSize.Y));
+	if (Size.Equals(NewSize, 0.001f))
+	{
+		return;
+	}
+
+	Modify();
+	Size = NewSize;
+	SizeChangedEvent.Broadcast(Size);
+	NotifyGraphChanged();
+}
+
+bool FScribbleNode::Serialize(FArchive& Ar)
+{
+	Ar.UsingCustomVersion(FScribbleObjectVersion::GUID);
+
+	if (Ar.IsSaving())
+	{
+		// update the position using the potential anchor
+		Position = GetPosition(true);
+	}
+	
+	Ar << Guid;
+	Ar << Anchor;
+	Ar << Position;
+	Ar << RelativePosition;
+	Ar << Size;
+	return true;
+}
+
+bool FScribbleNode::Identical(const FScribbleNode* Other, uint32 PortFlags) const
+{
+	return
+		Type == Other->Type &&
+		Guid == Other->Guid &&
+		Anchor == Other->Anchor &&
+		(Position - Other->Position).IsNearlyZero() &&
+		(RelativePosition - Other->RelativePosition).IsNearlyZero() &&
+		(Size - Other->Size).IsNearlyZero();
+}
+
+#if WITH_EDITOR
+
+int32 FScribbleNode::OnPaint(const FGeometry& InAllottedGeometry, FSlateWindowElementList& OutDrawElements, int32 InLayerId,
+	const FWidgetStyle& InWidgetStyle, bool bSelected, float InZoomAmount) const
+{
+	return InLayerId;
+}
+
+#endif
+
+const FScribbleGraphData* FScribbleNode::GetGraph() const
+{
+	if (WeakGraph.IsValid())
+	{
+		return WeakGraph.Pin().Get();
+	}
+	return nullptr;
+}
+
+FScribbleGraphData* FScribbleNode::GetGraph()
+{
+	if (WeakGraph.IsValid())
+	{
+		return WeakGraph.Pin().Get();
+	}
+	return nullptr;
+}
+
+void FScribbleNode::Modify()
+{
+	if (FScribbleGraphData* Graph = GetGraph())
+	{
+		Graph->Modify();
+	}
+}
+
+void FScribbleNode::IncrementChangeBracket()
+{
+	ChangeBracket++;
+}
+
+void FScribbleNode::DecrementChangeBracket()
+{
+	ChangeBracket = FMath::Max(ChangeBracket - 1, 0);
+	NotifyGraphChanged();
+}
+
+void FScribbleNode::NotifyGraphChanged()
+{
+	if (ChangeBracket == 0 && AccumulatedChanges > 0)
+	{
+		if (FScribbleGraphData* Graph = GetGraph())
+		{
+			Graph->NotifyChanged();
+		}
+		AccumulatedChanges = 0;
+	}
+	else
+	{
+		AccumulatedChanges++;
+	}
+}
+
+bool FScribbleLineStrip::Serialize(FArchive& Ar)
+{
+	Ar << Vertices;
+	Ar << Thickness;
+	Ar << Color;
+	return true;
+}
+
+bool FScribbleLineStrip::Identical(const FScribbleLineStrip* Other, uint32 PortFlags) const
+{
+	if (!Other)
+	{
+		return false;
+	}
+
+	if (Vertices.Num() != Other->Vertices.Num())
+	{
+		return false;
+	}
+	if (!FMath::IsNearlyEqual(Thickness, Other->Thickness))
+	{
+		return false;
+	}
+	if (!Color.Equals(Other->Color))
+	{
+		return false;
+	}
+	for(int32 VertexIndex = 0; VertexIndex < Vertices.Num(); VertexIndex++)
+	{
+		if (!Vertices[VertexIndex].Equals(Other->Vertices[VertexIndex]))
+		{
+			return false;
+		}
+	}
+	
+	return true;
+}
+
+bool FLineStripScribbleNode::Serialize(FArchive& Ar)
+{
+	if (!FScribbleNode::Serialize(Ar))
+	{
+		return false;
+	}
+	Ar << LineStrips;
+	return true;
+}
+
+bool FLineStripScribbleNode::Identical(const FScribbleNode* Other, uint32 PortFlags) const
+{
+	if (!FScribbleNode::Identical(Other, PortFlags))
+	{
+		return false;
+	}
+
+	const FLineStripScribbleNode* OtherNode = CastChecked<FLineStripScribbleNode>(Other);
+	if(LineStrips.Num() != OtherNode->LineStrips.Num())
+	{
+		return false;
+	}
+
+	for (int32 LineStripIndex = 0; LineStripIndex < LineStrips.Num(); LineStripIndex++)
+	{
+		const FScribbleLineStrip& LineStrip = LineStrips[LineStripIndex];
+		const FScribbleLineStrip& OtherLineStrip = OtherNode->LineStrips[LineStripIndex];
+		if (!LineStrip.Identical(&OtherLineStrip, PortFlags))
+		{
+			return false;
+		}
+	}
+	
+	return true;
+}
+
+void FLineStripScribbleNode::OffsetPosition(const FVector2f& InNewPosition)
+{
+	const FVector2f CurrentPosition = GetPosition();
+	if (InNewPosition.Equals(CurrentPosition, 0.001f))
+	{
+		return;
+	}
+
+	Modify();
+	IncrementChangeBracket();
+
+	FVector2f NewSize = FVector2f(1,1);
+	const FVector2f Delta = CurrentPosition - InNewPosition;
+	for (FScribbleLineStrip& LineStrip : LineStrips)
+	{
+		for (FVector2f& Vertex : LineStrip.Vertices)
+		{
+			Vertex += Delta;
+			NewSize.X = FMath::Max<float>(NewSize.X, Vertex.X);
+			NewSize.Y = FMath::Max<float>(NewSize.Y, Vertex.Y);
+		}
+	}
+
+	SetPosition(InNewPosition);
+	SetSize(NewSize);
+
+	DecrementChangeBracket();
+}
+
+bool FLineStripScribbleNode::DownSample(float InPrecision)
+{
+	Modify();
+	IncrementChangeBracket();
+
+	TArray<FScribbleLineStrip> NewLineStrips;
+	for (FScribbleLineStrip& LineStrip : LineStrips)
+	{
+		FScribbleLineStrip NewLineStrip;
+		NewLineStrip.Color = LineStrip.Color;
+		NewLineStrip.Thickness = LineStrip.Thickness;
+
+		for (int32 Index = 0; Index < LineStrip.Vertices.Num(); Index++)
+		{
+			const FVector2f& Vertex = LineStrip.Vertices[Index];
+			if (NewLineStrip.Vertices.IsEmpty())
+			{
+				NewLineStrip.Vertices.Add(Vertex);
+				continue;
+			}
+
+			// filter out points too close to each other
+			const FVector2f& LastVertex = NewLineStrip.Vertices.Last();
+			if (FVector2f::Distance(Vertex, LastVertex) < InPrecision)
+			{
+				NotifyGraphChanged();
+				continue;
+			}
+
+			if (NewLineStrip.Vertices.Num() > 1)
+			{
+				// if the last vertex is between the new one and the one before
+				// they may form a straight line
+				const FVector2f& LastLastVertex = NewLineStrip.Vertices[NewLineStrip.Vertices.Num() - 2];
+				const FVector2f& PointOnSegment = FMath::ClosestPointOnSegment2D(LastVertex, LastLastVertex, Vertex);
+				if (FVector2f::Distance(LastVertex, PointOnSegment) < InPrecision)
+				{
+					// project the new vertex towards the last vertex's direction
+					const FVector2f LastVertexWithSameDistance = LastLastVertex + (LastVertex - LastLastVertex).GetSafeNormal() * FVector2f::Distance(Vertex, LastLastVertex);
+
+					// replace the last vertex with this one
+					NotifyGraphChanged();
+					NewLineStrip.Vertices.Last() = (LastVertexWithSameDistance + Vertex) * 0.5f;
+					continue;
+				}
+			}
+
+			NewLineStrip.Vertices.Add(Vertex);
+		}
+
+		// filter out empty line strips
+		if (NewLineStrip.Vertices.Num() <= 1)
+		{
+			NotifyGraphChanged();
+			continue;
+		}
+
+		// filter out line strips that are too small
+		FBox2f Bounds = FBox2f(EForceInit::ForceInit);
+		for (const FVector2f& Vertex : NewLineStrip.Vertices)
+		{
+			Bounds += Vertex;
+		}
+		if (Bounds.GetSize().X < InPrecision && Bounds.GetSize().Y < InPrecision)
+		{
+			NotifyGraphChanged();
+			continue;
+		}
+		
+		NewLineStrips.Add(NewLineStrip);
+	}
+
+	Swap(LineStrips, NewLineStrips);
+	DecrementChangeBracket();
+
+	return !LineStrips.IsEmpty();
+}
+
+FBox2f FLineStripScribbleNode::GetContentBounds() const
+{
+	FBox2f Bounds(EForceInit::ForceInit);
+	for (const FScribbleLineStrip& LineStrip : LineStrips)
+	{
+		for (const FVector2f& Vertex : LineStrip.Vertices)
+		{
+			Bounds += Vertex;
+		}
+	}
+	return Bounds;
+}
+
+#if WITH_EDITOR
+
+int32 FLineStripScribbleNode::OnPaint(const FGeometry& InAllottedGeometry, FSlateWindowElementList& OutDrawElements, int32 InLayerId,
+                                      const FWidgetStyle& InWidgetStyle, bool bSelected, float InZoomAmount) const
+{
+	int32 LastLayerId = FScribbleNode::OnPaint(InAllottedGeometry, OutDrawElements, InLayerId, InWidgetStyle, bSelected, InZoomAmount);
+	
+	const FPaintGeometry PaintGeometry = InAllottedGeometry.ToPaintGeometry();
+	
+	TOptional<FSlateResourceHandle> WhiteBrushResourceHandle;
+	for (int32 LineStripIndex = 0; LineStripIndex < LineStrips.Num(); LineStripIndex++)
+	{
+		const FScribbleLineStrip& LineStrip = LineStrips[LineStripIndex];
+		if (LineStrip.Vertices.Num() < 2)
+		{
+			continue;
+		}
+		FLinearColor Color = LineStrip.Color;
+		if (bSelected)
+		{
+			Color = UScribbleEditorSettings::Get()->SelectionColor;
+		}
+
+		const float Thickness = LineStrip.Thickness * InZoomAmount;
+		if (Thickness < 0.75f)
+		{
+			continue;
+		}
+		
+		FSlateDrawElement::MakeLines(OutDrawElements, ++LastLayerId, PaintGeometry, LineStrip.Vertices, ESlateDrawEffect::None, Color, false, Thickness);
+
+		if (Thickness > 4.f && LineStrip.Vertices.Num() > 2)
+		{
+			VertexBuffer.Reset();
+			IndexBuffer.Reset();
+			VertexBuffer.Reserve((LineStrip.Vertices.Num() - 1) * 3);
+			IndexBuffer.Reserve((LineStrip.Vertices.Num() - 1) * 3);
+
+			const FSlateRenderTransform& Transform = PaintGeometry.GetAccumulatedRenderTransform();
+
+			const float HalfThickness = LineStrip.Thickness * 0.5f;
+			const FColor CustomVertColor = Color.ToFColor(true);
+			const FVector2f* VertexA = &LineStrip.Vertices[0];
+			const FVector2f* VertexB = VertexA + 1;
+			const FVector2f* VertexC = VertexB + 1;
+			for (int32 Index = 2; Index < LineStrip.Vertices.Num(); Index++)
+			{
+				// todo: this likely has to be optimized. we should potentially also cache
+				// of these results for a known last paint geometry
+				const FVector2f DeltaABUnit = ((*VertexA) - (*VertexB)).GetSafeNormal();
+				const FVector2f NormalAB(-DeltaABUnit.Y, DeltaABUnit.X);
+				const FVector2f DeltaCBUnit = ((*VertexC) - (*VertexB)).GetSafeNormal();
+				const FVector2f NormalCB(-DeltaCBUnit.Y, DeltaCBUnit.X);
+				
+				const float Dot = FMath::Abs<float>(DeltaABUnit.Dot(DeltaCBUnit));
+				if (Dot < 1.f)
+				{
+					float FlipA = 1.f;
+					if (DeltaCBUnit.Dot(NormalAB) > 0)
+					{
+						FlipA = -1.f;
+					}
+					float FlipB = 1.f;
+					if (DeltaABUnit.Dot(NormalCB) > 0)
+					{
+						FlipB = -1.f;
+					}
+
+					FVector2f CornerB1 = (*VertexB) + NormalAB * HalfThickness * FlipA;
+					FVector2f CornerB2 = (*VertexB) + NormalCB * HalfThickness * FlipB;
+					FVector2f CornerB3 = (*VertexB) - NormalAB * HalfThickness * FlipA;
+					FVector2f CornerB4 = (*VertexB) - NormalCB * HalfThickness * FlipB;
+
+					IndexBuffer.Add(VertexBuffer.Add(FSlateVertex::Make<ESlateVertexRounding::Disabled>(
+						Transform,
+						CornerB1,
+						{0.f,0.f},
+						CustomVertColor
+					)));
+					IndexBuffer.Add(VertexBuffer.Add(FSlateVertex::Make<ESlateVertexRounding::Disabled>(
+						Transform,
+						CornerB2,
+						{0.f,0.f},
+						CustomVertColor
+					)));
+					IndexBuffer.Add(VertexBuffer.Add(FSlateVertex::Make<ESlateVertexRounding::Disabled>(
+						Transform,
+						FMath::Lerp(CornerB3, CornerB4, 0.5f),
+						{0.f,0.f},
+						CustomVertColor
+					)));
+				}
+				
+				VertexA = VertexB;
+				VertexB = VertexC;
+				VertexC++;
+			}
+
+			if (!VertexBuffer.IsEmpty() && !IndexBuffer.IsEmpty())
+			{
+				if (!WhiteBrushResourceHandle.IsSet())
+				{
+					WhiteBrushResourceHandle = FSlateApplication::Get().GetRenderer()->GetResourceHandle(*FAppStyle::GetBrush("WhiteBrush"));
+				}
+				FSlateDrawElement::MakeCustomVerts(OutDrawElements, LastLayerId, WhiteBrushResourceHandle.GetValue(), VertexBuffer, IndexBuffer, nullptr, 0, 0);
+			}
+		}
+	}
+	return LastLayerId;
+}
+
+bool FLineStripScribbleNode::AppendMouseEvent(const FGeometry& MyGeometry, const FPointerEvent& MouseEvent, const FVector2f& InPanelPosition)
+{
+	if (LineStrips.IsEmpty())
+	{
+		return false;
+	}
+
+	const FVector2f PositionInNode = InPanelPosition - GetPosition();
+
+	FScribbleLineStrip& LineStrip = LineStrips.Last();
+	if (!LineStrip.Vertices.IsEmpty())
+	{
+		if (LineStrip.Vertices.Last().Equals(PositionInNode, 1.f))
+		{
+			return false;
+		}
+	}
+
+	LineStrip.Vertices.Add(PositionInNode);
+	SetSize({
+		FMath::Max<float>(Size.X, PositionInNode.X),
+		FMath::Max<float>(Size.Y, PositionInNode.Y)
+	});
+	return true;
+}
+
+bool FLineStripScribbleNode::IntersectsCursorPosition(const FVector2f& InLocalPosition) const
+{
+	for (const FScribbleLineStrip& LineStrip : LineStrips)
+	{
+		for (int32 Index = 1; Index < LineStrip.Vertices.Num(); Index++)
+		{
+			const FVector2f& VertexA = LineStrip.Vertices[Index-1];
+			const FVector2f& VertexB = LineStrip.Vertices[Index];
+			const FVector2f& PointOnSegment = FMath::ClosestPointOnSegment2D(InLocalPosition, VertexA, VertexB);
+			const float Distance = FVector2f::Distance(InLocalPosition, PointOnSegment);
+			if (Distance < LineStrip.Thickness + 8.f)
+			{
+				return true;
+			}
+		}
+	}
+	return false;
+}
+
+#endif // WITH_EDITOR
+
+#undef LOCTEXT_NAMESPACE
diff --git a/Engine/Plugins/Experimental/Animation/Scribble/Source/Scribble/Private/ScribbleSettings.cpp b/Engine/Plugins/Experimental/Animation/Scribble/Source/Scribble/Private/ScribbleSettings.cpp
new file mode 100644
--- /dev/null
+++ b/Engine/Plugins/Experimental/Animation/Scribble/Source/Scribble/Private/ScribbleSettings.cpp
@@ -0,0 +1,31 @@
+// Copyright Epic Games, Inc. All Rights Reserved.
+
+#include "ScribbleSettings.h"
+
+#include UE_INLINE_GENERATED_CPP_BY_NAME(ScribbleSettings)
+
+UScribbleEditorSettings::UScribbleEditorSettings(const FObjectInitializer& ObjectInitializer)
+    : Super(ObjectInitializer)
+#if WITH_EDITORONLY_DATA
+    , Color(FLinearColor::Blue)
+	, SelectionColor(FColor::FromHex(TEXT("515151FF")))
+	, AnchorColor(FColor::FromHex(TEXT("515151FF")))
+    , Thickness(4.f)
+    , Precision(3.f)
+#endif
+{
+}
+
+#if WITH_EDITOR
+
+float UScribbleEditorSettings::GetSmoothing() const
+{
+	return FMath::Clamp((Precision - 2.f) / 12.f, 0.0f, 1.0f);
+}
+
+void UScribbleEditorSettings::SetSmoothing(float InValue)
+{
+	Precision = 2.f + InValue * 12.f;
+}
+
+#endif
diff --git a/Engine/Plugins/Experimental/Animation/Scribble/Source/Scribble/Public/ScribbleGraph.h b/Engine/Plugins/Experimental/Animation/Scribble/Source/Scribble/Public/ScribbleGraph.h
new file mode 100644
--- /dev/null
+++ b/Engine/Plugins/Experimental/Animation/Scribble/Source/Scribble/Public/ScribbleGraph.h
@@ -0,0 +1,142 @@
+// Copyright Epic Games, Inc. All Rights Reserved.
+
+#pragma once
+
+#include "ScribbleNode.h"
+#include "HAL/IConsoleManager.h"
+#include "ScribbleGraph.generated.h"
+
+#define UE_API SCRIBBLE_API
+
+extern UE_API TAutoConsoleVariable<bool> CVarScribbleEnabled;
+
+DECLARE_DELEGATE_RetVal(FName, FScribbleGetCurrentAnchor)
+DECLARE_DELEGATE_RetVal_OneParam(TOptional<FVector2f>, FScribbleResolveAnchor, const FName&)
+
+/**
+ * The backing data for a scribble graph. Given we want this to be used as a TSharedPtr 
+ * it's separated from the UStruct graph.
+ * This is used for the data backend only and corresponds with
+ * the UScribbleEdGraphNode as well as the SScribbleGraphNode.
+ */
+struct FScribbleGraphData : public TSharedFromThis<FScribbleGraphData>
+{
+public:
+
+	FScribbleGraphData()
+		: ViewOffset(FVector2f::ZeroVector)
+		, ZoomAmount(1.f)
+		, ChangeBracket(0)
+		, AccumulatedChanges(0)
+	{
+	}
+
+	UE_API void Reset();
+	UE_API bool IsEmpty() const;
+	UE_API int32 NumNodes() const;
+	UE_API FScribbleNode* GetNode(int32 InNodeIndex);
+	UE_API const FScribbleNode* GetNode(int32 InNodeIndex) const;
+	UE_API TSharedPtr<FScribbleNode> GetNodePtr(int32 InNodeIndex) const;
+	UE_API FScribbleNode* FindNode(const FGuid& InId);
+	UE_API const FScribbleNode* FindNode(const FGuid& InId) const;
+	UE_API TSharedPtr<FScribbleNode> FindNodePtr(const FGuid& InId) const;
+	UE_API FScribbleNode* operator[](int32 InNodeIndex);
+	UE_API const FScribbleNode* operator[](int32 InNodeIndex) const;
+	TArray<TSharedPtr<FScribbleNode>>::RangedForIteratorType begin() { return Nodes.begin(); }
+	TArray<TSharedPtr<FScribbleNode>>::RangedForIteratorType end() { return Nodes.end(); }
+	TArray<TSharedPtr<FScribbleNode>>::RangedForConstIteratorType begin() const { return Nodes.begin(); }
+	TArray<TSharedPtr<FScribbleNode>>::RangedForConstIteratorType end() const { return Nodes.end(); }
+
+	UE_API FGuid AddNode(const TSharedPtr<FScribbleNode>& InNode);
+	UE_API bool RemoveNode(const TSharedPtr<FScribbleNode>& InNode);
+	UE_API bool RemoveNode(const FGuid& InId);
+
+	UE_API bool SupportsAnchors() const;
+	UE_API FName GetCurrentAnchor() const;
+	UE_API TOptional<FVector2f> ResolveAnchor(const FName& InName) const;
+
+	UE_API TSharedPtr<FScribbleNode> GroupNodes(const TArray<TSharedPtr<FScribbleNode>>& InNodes);
+	UE_API TArray<TSharedPtr<FScribbleNode>> UngroupNode(const TSharedPtr<FScribbleNode>& InNode);
+
+	FSimpleMulticastDelegate& OnModify() { return OnModifyEvent; }
+	FSimpleMulticastDelegate& OnChanged() { return OnChangedEvent; }
+	FScribbleGetCurrentAnchor& OnGetCurrentAnchor() { return GetCurrentAnchorDelegate; }
+	FScribbleResolveAnchor& OnResolveAnchor() { return ResolveAnchorDelegate; }
+
+	UE_API bool Serialize(FArchive& Ar);
+	UE_API void Modify();
+
+	const FVector2f& GetViewOffset() const { return ViewOffset; }
+	float GetZoomAmount() const { return ZoomAmount; }
+	UE_API void SetView(const FVector2f& InViewOffset, float InZoomAmount);
+	
+protected:
+
+	UE_API void IncrementChangeBracket();
+	UE_API void DecrementChangeBracket();
+	UE_API void NotifyChanged();
+
+	FVector2f ViewOffset;
+	float ZoomAmount;
+
+	TArray<TSharedPtr<FScribbleNode>> Nodes;
+	TMap<FGuid, TSharedPtr<FScribbleNode>> GuidToNode;
+	FSimpleMulticastDelegate OnModifyEvent;
+	FSimpleMulticastDelegate OnChangedEvent;
+	FScribbleGetCurrentAnchor GetCurrentAnchorDelegate;
+	FScribbleResolveAnchor ResolveAnchorDelegate;
+
+	int32 ChangeBracket;
+	int32 AccumulatedChanges;
+
+	friend struct FScribbleNode;
+	friend class UScribbleEdGraph;
+};
+
+/**
+ * The Scribbe graph struct used mainly for serialization and as a data type for UPROPERTY macros.
+ */
+USTRUCT()
+struct FScribbleGraph
+{
+	GENERATED_BODY()
+
+public:
+
+	FScribbleGraph()
+	{
+	}
+
+	UE_API void Reset();
+	UE_API bool IsValid() const;
+	UE_API FScribbleGraphData* GetData();
+	UE_API const FScribbleGraphData* GetData() const;
+	TSharedPtr<FScribbleGraphData> GetDataPtr() const { return Data; }
+	UE_API TSharedPtr<FScribbleGraphData> GetOrCreateDataPtr();
+
+	UE_API bool Serialize(FArchive& Ar);
+	bool Identical(const FScribbleGraph* Other, uint32 PortFlags) const;
+
+	friend FArchive& operator<<(FArchive& Ar, FScribbleGraph& P)
+	{
+		P.Serialize(Ar);
+		return Ar;
+	}
+
+private:
+
+	TSharedPtr<FScribbleGraphData> Data;
+};
+
+template<>
+struct TStructOpsTypeTraits<FScribbleGraph> : public TStructOpsTypeTraitsBase2<FScribbleGraph>
+{
+	enum 
+	{
+		WithSerializer = true, // struct has a Serialize function for serializing its state to an FArchive.
+		WithIdentical = true, // struct can be compared via an Identical(const T* Other, uint32 PortFlags) function.
+	};
+	static constexpr EPropertyObjectReferenceType WithSerializerObjectReferences = EPropertyObjectReferenceType::None;
+};
+
+#undef UE_API
diff --git a/Engine/Plugins/Experimental/Animation/Scribble/Source/Scribble/Public/ScribbleNode.h b/Engine/Plugins/Experimental/Animation/Scribble/Source/Scribble/Public/ScribbleNode.h
new file mode 100644
--- /dev/null
+++ b/Engine/Plugins/Experimental/Animation/Scribble/Source/Scribble/Public/ScribbleNode.h
@@ -0,0 +1,271 @@
+// Copyright Epic Games, Inc. All Rights Reserved.
+
+#pragma once
+
+#include "CoreMinimal.h"
+
+#if WITH_EDITOR
+#include "Input/Events.h"
+#include "Layout/PaintGeometry.h"
+#include "Rendering/DrawElements.h"
+#include "Styling/WidgetStyle.h"
+#endif
+
+#include "ScribbleNode.generated.h"
+
+struct FScribbleGraphData;
+
+#define UE_API SCRIBBLE_API
+
+DECLARE_MULTICAST_DELEGATE_OneParam(FScribbleNodeDimensionsChanged, const FVector2f&);
+
+#define DECLARE_SCRIBBLE_NODE_METHODS(ScribbleNodeType) \
+virtual UScriptStruct* GetScriptStruct() const override { return ScribbleNodeType::StaticStruct(); } \
+template<typename T> \
+friend const T* Cast(const ScribbleNodeType* InNode) \
+{ \
+return Cast<T>((const FScribbleNode*) InNode); \
+} \
+template<typename T> \
+friend T* Cast(ScribbleNodeType* InNode) \
+{ \
+return Cast<T>((FScribbleNode*) InNode); \
+} \
+template<typename T> \
+friend const T* CastChecked(const ScribbleNodeType* InNode) \
+{ \
+return CastChecked<T>((const FScribbleNode*) InNode); \
+} \
+template<typename T> \
+friend T* CastChecked(ScribbleNodeType* InNode) \
+{ \
+return CastChecked<T>((FScribbleNode*) InNode); \
+} \
+static bool IsClassOf(const FScribbleNode* InNode) \
+{ \
+return InNode->GetScriptStruct()->IsChildOf(StaticStruct()); \
+}
+
+UENUM()
+namespace EScribbleNodeType
+{
+	enum Type : uint8
+	{
+		Invalid,
+		LineStrip,
+	};
+}
+
+/*
+ * A node in a scribble graph. This is used for the data backend only and corresponds with
+ * the UScribbleEdGraphNode as well as the SScribbleGraphNode.
+ */
+USTRUCT()
+struct FScribbleNode
+{
+	GENERATED_BODY()
+
+public:
+
+	FScribbleNode(const EScribbleNodeType::Type& InType = EScribbleNodeType::Invalid)
+		: Type(InType)
+		, Guid(FGuid::NewGuid())
+		, Anchor(NAME_None)
+		, Position(FVector2f::ZeroVector)
+		, RelativePosition(FVector2f::ZeroVector)
+		, Size(1,1)
+		, ChangeBracket(0)
+		, AccumulatedChanges(0)
+	{
+	}
+
+	virtual ~FScribbleNode()
+	{
+	}
+
+	const EScribbleNodeType::Type GetType() const { return Type; }
+	const FGuid& GetId() const { return Guid; }
+	const FName& GetAnchor() const { return Anchor; }
+	UE_API void SetAnchor(const FName& InAnchor);
+
+	UE_API FVector2f GetPosition(bool bResolveAnchor = true) const;
+	const FVector2f& GetRelativePosition() const { return RelativePosition; }
+	UE_API void SetPosition(const FVector2f& InPosition);
+	virtual void OffsetPosition(const FVector2f& InNewPosition) {}
+	virtual FBox2f GetContentBounds() const { return FBox2f(EForceInit::ForceInitToZero); }
+	const FVector2f& GetSize() const { return Size; }
+	UE_API void SetSize(const FVector2f& InSize);
+	virtual bool DownSample(float InPrecision) { return false; }
+	
+	UE_API virtual bool Serialize(FArchive& Ar);
+	UE_API virtual bool Identical(const FScribbleNode* Other, uint32 PortFlags) const;
+
+	FScribbleNodeDimensionsChanged& OnPositionChanged() { return PositionChangedEvent; }
+	FScribbleNodeDimensionsChanged& OnSizeChanged() { return SizeChangedEvent; }
+
+#if WITH_EDITOR
+	virtual int32 OnPaint(const FGeometry& InAllottedGeometry, FSlateWindowElementList& OutDrawElements, int32 InLayerId, const FWidgetStyle& InWidgetStyle, bool bSelected, float InZoomAmount) const;
+	virtual bool AppendMouseEvent(const FGeometry& MyGeometry, const FPointerEvent& MouseEvent, const FVector2f& InPanelPosition) { return false; }
+	virtual bool IntersectsCursorPosition(const FVector2f& InLocalPosition) const { return false; }
+#endif
+
+	static bool IsClassOf(const FScribbleNode* InNode)
+	{
+		return true;
+	}
+
+	virtual UScriptStruct* GetScriptStruct() const { return FScribbleNode::StaticStruct(); }
+
+	bool IsA(const UScriptStruct* InScriptStruct) const
+	{
+		return GetScriptStruct()->IsChildOf(InScriptStruct);
+	}
+
+	template<typename T>
+	bool IsA() const { return T::IsClassOf(this); }
+
+	template<typename T>
+	friend const T* Cast(const FScribbleNode* InNode)
+	{
+		if(InNode)
+		{
+			if(InNode->IsA<T>())
+			{
+				return static_cast<const T*>(InNode);
+			}
+		}
+		return nullptr;
+	}
+
+	template<typename T>
+	friend T* Cast(FScribbleNode* InNode)
+	{
+		if(InNode)
+		{
+			if(InNode->IsA<T>())
+			{
+				return static_cast<T*>(InNode);
+			}
+		}
+		return nullptr;
+	}
+
+	template<typename T>
+	friend const T* CastChecked(const FScribbleNode* InNode)
+	{
+		const T* Node = Cast<T>(InNode);
+		check(Node);
+		return Node;
+	}
+
+	template<typename T>
+	friend T* CastChecked(FScribbleNode* InNode)
+	{
+		T* Node = Cast<T>(InNode);
+		check(Node);
+		return Node;
+	}
+
+	UE_API const FScribbleGraphData* GetGraph() const;
+	UE_API FScribbleGraphData* GetGraph();
+
+protected:
+
+	UE_API void Modify();
+	UE_API void IncrementChangeBracket();
+	UE_API void DecrementChangeBracket();
+	UE_API void NotifyGraphChanged();
+
+	void UpdateRelativePosition(const FVector2f& InAbsolutePosition);
+
+	UPROPERTY()
+	TEnumAsByte<EScribbleNodeType::Type> Type; 
+
+	UPROPERTY()
+	FGuid Guid;
+
+	UPROPERTY()
+	FName Anchor;
+
+	UPROPERTY()
+	FVector2f Position;
+
+	UPROPERTY()
+	FVector2f RelativePosition;
+
+	UPROPERTY()
+	FVector2f Size;
+
+	TWeakPtr<FScribbleGraphData> WeakGraph;
+	FScribbleNodeDimensionsChanged PositionChangedEvent;
+	FScribbleNodeDimensionsChanged SizeChangedEvent;
+
+	int32 ChangeBracket;
+	int32 AccumulatedChanges;
+
+	friend struct FScribbleGraphData;
+	friend class SScribbleGraphPanel;
+};
+
+USTRUCT()
+struct FScribbleLineStrip
+{
+	GENERATED_BODY()
+
+public:
+
+	FScribbleLineStrip()
+		: Thickness(1.f)
+		, Color(FLinearColor::White)
+	{
+	}
+
+	UPROPERTY(EditAnywhere, Category=Scribble)
+	TArray<FVector2f> Vertices;
+
+	UPROPERTY(EditAnywhere, Category=Scribble)
+	float Thickness;
+
+	UPROPERTY(EditAnywhere, Category=Scribble)
+	FLinearColor Color;
+
+	UE_API bool Serialize(FArchive& Ar);
+	UE_API bool Identical(const FScribbleLineStrip* Other, uint32 PortFlags) const;
+
+	friend FArchive& operator<<(FArchive& Ar, FScribbleLineStrip& P)
+	{
+		P.Serialize(Ar);
+		return Ar;
+	}
+};
+
+USTRUCT()
+struct UE_API FLineStripScribbleNode : public FScribbleNode
+{
+	GENERATED_BODY()
+	DECLARE_SCRIBBLE_NODE_METHODS(FLineStripScribbleNode)
+
+	FLineStripScribbleNode()
+	: FScribbleNode(EScribbleNodeType::LineStrip)
+	{
+	}
+
+	UPROPERTY(EditAnywhere, Category=Scribble)
+	TArray<FScribbleLineStrip> LineStrips;
+
+	virtual bool Serialize(FArchive& Ar) override;
+	virtual bool Identical(const FScribbleNode* Other, uint32 PortFlags) const;
+	virtual void OffsetPosition(const FVector2f& InNewPosition) override;
+	virtual bool DownSample(float InPrecision) override;
+	virtual FBox2f GetContentBounds() const override;;
+
+#if WITH_EDITOR
+	virtual int32 OnPaint(const FGeometry& InAllottedGeometry, FSlateWindowElementList& OutDrawElements, int32 InLayerId, const FWidgetStyle& InWidgetStyle, bool bSelected, float InZoomAmount) const override;
+	virtual bool AppendMouseEvent(const FGeometry& MyGeometry, const FPointerEvent& MouseEvent, const FVector2f& InPanelPosition) override;
+	virtual bool IntersectsCursorPosition(const FVector2f& InLocalPosition) const override;
+	mutable TArray<FSlateVertex> VertexBuffer;
+	mutable TArray<SlateIndex> IndexBuffer;
+#endif
+};
+
+#undef UE_API
diff --git a/Engine/Plugins/Experimental/Animation/Scribble/Source/Scribble/Public/ScribbleObjectVersion.h b/Engine/Plugins/Experimental/Animation/Scribble/Source/Scribble/Public/ScribbleObjectVersion.h
new file mode 100644
--- /dev/null
+++ b/Engine/Plugins/Experimental/Animation/Scribble/Source/Scribble/Public/ScribbleObjectVersion.h
@@ -0,0 +1,29 @@
+// Copyright Epic Games, Inc. All Rights Reserved.
+#pragma once
+
+#include "CoreTypes.h"
+#include "Misc/Guid.h"
+
+#define UE_API SCRIBBLE_API
+
+// Custom serialization version for scribble plugin
+struct FScribbleObjectVersion
+{
+	enum Type
+	{
+		// Before any version changes were made
+		BeforeCustomVersionWasAdded,
+
+		// -----<new versions can be added above this line>-------------------------------------------------
+		VersionPlusOne,
+		LatestVersion = VersionPlusOne - 1
+	};
+
+	// The GUID for this custom version number
+	UE_API const static FGuid GUID;
+
+private:
+	FScribbleObjectVersion() {}
+};
+
+#undef UE_API
diff --git a/Engine/Plugins/Experimental/Animation/Scribble/Source/Scribble/Public/ScribbleSettings.h b/Engine/Plugins/Experimental/Animation/Scribble/Source/Scribble/Public/ScribbleSettings.h
new file mode 100644
--- /dev/null
+++ b/Engine/Plugins/Experimental/Animation/Scribble/Source/Scribble/Public/ScribbleSettings.h
@@ -0,0 +1,72 @@
+// Copyright Epic Games, Inc. All Rights Reserved.
+
+/*=============================================================================
+	ScribbleSettings.h: Declares the ScribbleSettings class.
+=============================================================================*/
+
+#pragma once
+
+#include "CoreMinimal.h"
+#include "Engine/DeveloperSettings.h"
+
+#include "ScribbleSettings.generated.h"
+
+#define UE_API SCRIBBLE_API
+
+/**
+ * Default Scribble settings.
+ */
+UCLASS(config = EditorPerProjectUserSettings, meta=(DisplayName="Scribble"), MinimalAPI)
+class UScribbleEditorSettings : public UDeveloperSettings
+{
+	GENERATED_UCLASS_BODY()
+
+public:
+
+#if WITH_EDITORONLY_DATA
+
+	// The color used for shapes
+	UPROPERTY(EditAnywhere, config, Category = Scribble)
+	FLinearColor Color;
+
+	// The color used to represent selection
+	UPROPERTY(EditAnywhere, config, Category = Scribble)
+	FLinearColor SelectionColor;
+
+	// The color used to draw the anchor line
+	UPROPERTY(EditAnywhere, config, Category = Scribble)
+	FLinearColor AnchorColor;
+
+	// The thickness used for lines
+	UPROPERTY(EditAnywhere, config, Category = Scribble)
+	float Thickness;
+
+	// Defines how exact lines are - and how small the smallest possible node is going to be
+	UPROPERTY(EditAnywhere, config, Category = Scribble)
+	float Precision;
+
+#endif
+
+#if WITH_EDITOR
+
+	UE_API float GetSmoothing() const;
+	UE_API void SetSmoothing(float InValue);
+
+	void PostEditChangeProperty(struct FPropertyChangedEvent& PropertyChangedEvent)
+	{
+		Super::PostEditChangeProperty(PropertyChangedEvent);
+
+		if (PropertyChangedEvent.ChangeType != EPropertyChangeType::Interactive)
+		{
+			// Dragging spinboxes causes this to be called every frame so we wait until they've finished dragging before saving.
+			SaveConfig();
+		}
+		
+		OnSettingChanged().Broadcast(this, PropertyChangedEvent);
+	}
+#endif
+
+	static UScribbleEditorSettings* Get() { return GetMutableDefault<UScribbleEditorSettings>(); }
+};
+
+#undef UE_API
diff --git a/Engine/Plugins/Experimental/Animation/Scribble/Source/Scribble/Scribble.Build.cs b/Engine/Plugins/Experimental/Animation/Scribble/Source/Scribble/Scribble.Build.cs
new file mode 100644
--- /dev/null
+++ b/Engine/Plugins/Experimental/Animation/Scribble/Source/Scribble/Scribble.Build.cs
@@ -0,0 +1,30 @@
+// Copyright Epic Games, Inc. All Rights Reserved.
+
+namespace UnrealBuildTool.Rules
+{
+    public class Scribble : ModuleRules
+    {
+        public Scribble(ReadOnlyTargetRules Target) : base(Target)
+        {
+			NumIncludedBytesPerUnityCPPOverride = 688128; // best unity size found from using UBT ProfileUnitySizes mode
+
+            PrivateDependencyModuleNames.AddRange(
+                new string[]
+                {
+                    "Core",
+                    "CoreUObject",
+                    "SlateCore",
+                    "Slate",
+                    "Engine",
+				}
+            );
+
+            PublicDependencyModuleNames.AddRange(
+                new string[]
+                {
+	                "DeveloperSettings",
+                }
+            );
+        }
+    }
+}
diff --git a/Engine/Plugins/Experimental/Animation/Scribble/Source/ScribbleEditor/Private/SScribbleGraph.cpp b/Engine/Plugins/Experimental/Animation/Scribble/Source/ScribbleEditor/Private/SScribbleGraph.cpp
new file mode 100644
--- /dev/null
+++ b/Engine/Plugins/Experimental/Animation/Scribble/Source/ScribbleEditor/Private/SScribbleGraph.cpp
@@ -0,0 +1,1214 @@
+// Copyright Epic Games, Inc. All Rights Reserved.
+
+#include "SScribbleGraph.h"
+
+#include "EdGraphSchema_K2_Actions.h"
+#include "Editor.h"
+#include "Editor/Transactor.h"
+#include "Editor/UnrealEdEngine.h"
+#include "Engine/Selection.h"
+#include "Settings/EditorStyleSettings.h"
+#include "ScribbleEdGraph.h"
+#include "ScribbleEdGraphSchema.h"
+#include "ScribbleEdGraphCommands.h"
+#include "ScribbleSettings.h"
+#include "SScribbleGraphNode.h"
+#include "UnrealEdGlobals.h"
+#include "Scribble/Public/ScribbleSettings.h"
+#include "Toolkits/GlobalEditorCommonCommands.h"
+#include "Widgets/Colors/SColorBlock.h"
+#include "Widgets/Colors/SColorPicker.h"
+#include "Widgets/Images/SImage.h"
+#include "Widgets/Input/SButton.h"
+#include "Widgets/Input/SSpinBox.h"
+#include "Widgets/SBoxPanel.h"
+
+#define LOCTEXT_NAMESPACE "SScribbleGraph"
+
+struct FScribbleZoomLevelsContainer : public FZoomLevelsContainer
+{
+	struct FScribbleZoomLevelEntry
+	{
+	public:
+		FScribbleZoomLevelEntry(float InZoomAmount, const FText& InDisplayText, EGraphRenderingLOD::Type InLOD)
+			: DisplayText(FText::Format(NSLOCTEXT("GraphEditor", "Zoom", "Zoom {0}"), InDisplayText))
+		, ZoomAmount(InZoomAmount)
+		, LOD(InLOD)
+		{
+		}
+
+	public:
+		FText DisplayText;
+		float ZoomAmount;
+		EGraphRenderingLOD::Type LOD;
+	};
+	
+	FScribbleZoomLevelsContainer()
+	{
+		ZoomLevels.Reserve(22);
+		ZoomLevels.Add(FScribbleZoomLevelEntry(0.025f, FText::FromString(TEXT("-14")), EGraphRenderingLOD::LowestDetail));
+		ZoomLevels.Add(FScribbleZoomLevelEntry(0.070f, FText::FromString(TEXT("-13")), EGraphRenderingLOD::LowestDetail));
+		ZoomLevels.Add(FScribbleZoomLevelEntry(0.100f, FText::FromString(TEXT("-12")), EGraphRenderingLOD::LowestDetail));
+		ZoomLevels.Add(FScribbleZoomLevelEntry(0.125f, FText::FromString(TEXT("-11")), EGraphRenderingLOD::LowestDetail));
+		ZoomLevels.Add(FScribbleZoomLevelEntry(0.150f, FText::FromString(TEXT("-10")), EGraphRenderingLOD::LowestDetail));
+		ZoomLevels.Add(FScribbleZoomLevelEntry(0.175f, FText::FromString(TEXT("-9")), EGraphRenderingLOD::LowestDetail));
+		ZoomLevels.Add(FScribbleZoomLevelEntry(0.200f, FText::FromString(TEXT("-8")), EGraphRenderingLOD::LowestDetail));
+		ZoomLevels.Add(FScribbleZoomLevelEntry(0.225f, FText::FromString(TEXT("-7")), EGraphRenderingLOD::LowDetail));
+		ZoomLevels.Add(FScribbleZoomLevelEntry(0.250f, FText::FromString(TEXT("-6")), EGraphRenderingLOD::LowDetail));
+		ZoomLevels.Add(FScribbleZoomLevelEntry(0.375f, FText::FromString(TEXT("-5")), EGraphRenderingLOD::MediumDetail));
+		ZoomLevels.Add(FScribbleZoomLevelEntry(0.500f, FText::FromString(TEXT("-4")), EGraphRenderingLOD::MediumDetail));
+		ZoomLevels.Add(FScribbleZoomLevelEntry(0.675f, FText::FromString(TEXT("-3")), EGraphRenderingLOD::MediumDetail));
+		ZoomLevels.Add(FScribbleZoomLevelEntry(0.750f, FText::FromString(TEXT("-2")), EGraphRenderingLOD::DefaultDetail));
+		ZoomLevels.Add(FScribbleZoomLevelEntry(0.875f, FText::FromString(TEXT("-1")), EGraphRenderingLOD::DefaultDetail));
+		ZoomLevels.Add(FScribbleZoomLevelEntry(1.000f, FText::FromString(TEXT("1:1")), EGraphRenderingLOD::DefaultDetail));
+		ZoomLevels.Add(FScribbleZoomLevelEntry(1.250f, FText::FromString(TEXT("+1")), EGraphRenderingLOD::DefaultDetail));
+		ZoomLevels.Add(FScribbleZoomLevelEntry(1.375f, FText::FromString(TEXT("+2")), EGraphRenderingLOD::DefaultDetail));
+		ZoomLevels.Add(FScribbleZoomLevelEntry(1.500f, FText::FromString(TEXT("+3")), EGraphRenderingLOD::FullyZoomedIn));
+		ZoomLevels.Add(FScribbleZoomLevelEntry(1.675f, FText::FromString(TEXT("+4")), EGraphRenderingLOD::FullyZoomedIn));
+		ZoomLevels.Add(FScribbleZoomLevelEntry(1.750f, FText::FromString(TEXT("+5")), EGraphRenderingLOD::FullyZoomedIn));
+		ZoomLevels.Add(FScribbleZoomLevelEntry(1.875f, FText::FromString(TEXT("+6")), EGraphRenderingLOD::FullyZoomedIn));
+		ZoomLevels.Add(FScribbleZoomLevelEntry(2.000f, FText::FromString(TEXT("+7")), EGraphRenderingLOD::FullyZoomedIn));
+	}
+
+	float GetZoomAmount(int32 InZoomLevel) const override
+	{
+		checkSlow(ZoomLevels.IsValidIndex(InZoomLevel));
+		return ZoomLevels[InZoomLevel].ZoomAmount;
+	}
+
+	int32 GetNearestZoomLevel(float InZoomAmount) const override
+	{
+		for (int32 ZoomLevelIndex=0; ZoomLevelIndex < GetNumZoomLevels(); ++ZoomLevelIndex)
+		{
+			if (InZoomAmount <= GetZoomAmount(ZoomLevelIndex))
+			{
+				return ZoomLevelIndex;
+			}
+		}
+
+		return GetDefaultZoomLevel();
+	}
+	
+	FText GetZoomText(int32 InZoomLevel) const override
+	{
+		checkSlow(ZoomLevels.IsValidIndex(InZoomLevel));
+		return ZoomLevels[InZoomLevel].DisplayText;
+	}
+	
+	int32 GetNumZoomLevels() const override
+	{
+		return ZoomLevels.Num();
+	}
+	
+	int32 GetDefaultZoomLevel() const override
+	{
+		return 14;
+	}
+
+	EGraphRenderingLOD::Type GetLOD(int32 InZoomLevel) const override
+	{
+		checkSlow(ZoomLevels.IsValidIndex(InZoomLevel));
+		return ZoomLevels[InZoomLevel].LOD;
+	}
+
+	TArray<FScribbleZoomLevelEntry> ZoomLevels;
+};
+
+
+SScribbleGraphPanel::SScribbleGraphPanel()
+: ScribbleGraphObj(nullptr)
+, bAllowNavigation(true)
+, RoundedFrameBrush(MakeUnique<FSlateRoundedBoxBrush>(FStyleColors::Transparent, 16.0f, FStyleColors::White, 2.0f))
+{
+}
+
+SScribbleGraphPanel::~SScribbleGraphPanel()
+{
+	if (!GExitPurge)
+	{
+		if ( ensure(ScribbleGraphObj) )
+		{
+			ScribbleGraphObj->RemoveFromRoot();
+		}
+	}
+}
+
+/** Constructs this widget with InArgs */
+void SScribbleGraphPanel::Construct(const FArguments& InArgs)
+{
+	// Create the graph
+	ScribbleGraphObj = NewObject<UScribbleEdGraph>();
+	ScribbleGraphObj->Schema = UScribbleEdGraphSchema::StaticClass();
+	ScribbleGraphObj->AddToRoot();
+	ScribbleGraphObj->Initialize(InArgs._GraphData);
+
+	ScribbleEnabled = InArgs._ScribbleEnabled;
+	ScribbleNodeType = InArgs._ScribbleNodeType;
+	ScribbleColor = InArgs._ScribbleColor;
+	if (!ScribbleColor.IsBound() && !ScribbleColor.IsSet())
+	{
+		ScribbleColor.BindLambda([]()
+		{
+			return UScribbleEditorSettings::Get()->Color;
+		});
+	}
+	ScribbleThickness = InArgs._ScribbleThickness;
+	if (!ScribbleThickness.IsBound() && !ScribbleThickness.IsSet())
+	{
+		ScribbleThickness.BindLambda([]()
+		{
+			return UScribbleEditorSettings::Get()->Thickness;
+		});
+	}
+	ScribblePrecision = InArgs._ScribblePrecision;
+	if (!ScribblePrecision.IsBound() && !ScribblePrecision.IsSet())
+	{
+		ScribblePrecision.BindLambda([]()
+		{
+			return UScribbleEditorSettings::Get()->Precision;
+		});
+	}
+	bAllowNavigation = InArgs._AllowNavigation;
+	WeakPanelToSync = InArgs._PanelToSync;
+
+	if (WeakPanelToSync.IsValid())
+	{
+		if (TSharedPtr<SGraphPanel> PanelToSync = WeakPanelToSync.Pin())
+		{
+			ZoomLevel = PanelToSync->GetZoomLevel();
+			SetZoomLevelsContainer(PanelToSync->GetZoomLevels());
+		}
+	}
+	else
+	{
+		ZoomLevel = FScribbleZoomLevelsContainer().GetDefaultZoomLevel();
+		SetZoomLevelsContainer<FScribbleZoomLevelsContainer>();
+	}
+
+	SGraphPanel::FArguments GraphPanelArgs;
+	GraphPanelArgs.GraphObj(ScribbleGraphObj);
+	GraphPanelArgs.OnSelectionChanged(this, &SScribbleGraphPanel::HandleSelectionChanged);
+	GraphPanelArgs.ShouldDrawBackground(InArgs._ShouldDrawBackground);
+	GraphPanelArgs.ShouldDrawSurroundingShadow(InArgs._ShouldDrawSurroundingShadow);
+	GraphPanelArgs.AllowZoom(bAllowNavigation);
+	GraphPanelArgs.AllowPanning(bAllowNavigation);
+	
+	SGraphPanel::Construct(GraphPanelArgs);
+
+	BindCommands();
+
+	Update();
+
+	GEditor->RegisterForUndo(this);
+}
+
+bool SScribbleGraphPanel::IsScribbleEnabled() const
+{
+	return ScribbleEnabled.Get(false);
+}
+
+bool SScribbleGraphPanel::HasActiveNodeType() const
+{
+	return GetActiveNodeType() != EScribbleNodeType::Invalid;
+}
+
+void SScribbleGraphPanel::ActivateSelectionMode()
+{
+	SetActiveNodeType(EScribbleNodeType::Invalid);
+}
+
+EScribbleNodeType::Type SScribbleGraphPanel::GetActiveNodeType() const
+{
+	return ScribbleNodeType.Get(EScribbleNodeType::Invalid);
+}
+
+void SScribbleGraphPanel::SetActiveNodeType(EScribbleNodeType::Type InNodeType)
+{
+	if (OnSetNodeType.IsBound())
+	{
+		OnSetNodeType.Execute(InNodeType);
+		return;
+	}
+	if (!ScribbleNodeType.IsBound())
+	{
+		ScribbleNodeType = InNodeType;
+	}
+}
+
+UScribbleEdGraph* SScribbleGraphPanel::GetScribbleEdGraph()
+{
+	return ScribbleGraphObj;
+}
+
+const UScribbleEdGraph* SScribbleGraphPanel::GetScribbleEdGraph() const
+{
+	return ScribbleGraphObj;
+}
+
+FScribbleGraphData* SScribbleGraphPanel::GetScribbleGraph()
+{
+	if (UScribbleEdGraph* ScribbleEdGraph = GetScribbleEdGraph())
+	{
+		return ScribbleEdGraph->GetGraphData();
+	}
+	return nullptr;
+}
+
+const FScribbleGraphData* SScribbleGraphPanel::GetScribbleGraph() const
+{
+	if (const UScribbleEdGraph* ScribbleEdGraph = GetScribbleEdGraph())
+	{
+		return ScribbleEdGraph->GetGraphData();
+	}
+	return nullptr;
+}
+
+FReply SScribbleGraphPanel::OnMouseButtonDown( const FGeometry& MyGeometry, const FPointerEvent& MouseEvent )
+{
+	if (!IsScribbleEnabled())
+	{
+		return FReply::Unhandled();
+	}
+	if (HasActiveNodeType())
+	{
+		if (MouseEvent.IsMouseButtonDown(EKeys::LeftMouseButton))
+		{
+			bIsScribbleDragging = true;
+			UpdateNewNode(MyGeometry, MouseEvent);
+			return FReply::Handled();
+		}
+	}
+	const FReply ReplyFromGraphPanel = SGraphPanel::OnMouseButtonDown(MyGeometry, MouseEvent);
+	if (ReplyFromGraphPanel.IsEventHandled())
+	{
+		if (NodeUnderMousePtr.IsValid())
+		{
+			return ReplyFromGraphPanel;
+		}
+	}
+
+	FArrangedChildren ArrangedChildren(EVisibility::Visible);
+	ArrangeChildNodes(MyGeometry, ArrangedChildren);
+
+	// use an alternative path to look up nodes to select by click.
+	// the nodes may be overlapping so the initial event may not catch it.
+	const FVector2f ArrangedSpacePosition = MouseEvent.GetScreenSpacePosition();
+	const int32 NumChildren = ArrangedChildren.Num();
+	for( int32 ChildIndex=NumChildren-1; ChildIndex >= 0; --ChildIndex )
+	{
+		const FArrangedWidget& Candidate = ArrangedChildren[ChildIndex];
+		if (Candidate.Geometry.IsUnderLocation( ArrangedSpacePosition ))
+		{
+			const FVector2f MousePositionInNode = Candidate.Geometry.AbsoluteToLocal(ArrangedSpacePosition);
+			const TSharedRef<SScribbleGraphNode>& ScribbleGraphNode = StaticCastSharedRef<SScribbleGraphNode, SWidget>(Candidate.Widget);
+			if (ScribbleGraphNode->CanBeSelected(MousePositionInNode))
+			{
+				// Track the node that we're dragging; we will move it in OnMouseMove.
+				this->OnBeginNodeInteraction(ScribbleGraphNode, MousePositionInNode);
+				return FReply::Handled().CaptureMouse( SharedThis(this) );
+			}
+		}
+	}
+
+	return ReplyFromGraphPanel;
+}
+
+void SScribbleGraphPanel::OnMouseCaptureLost(const FCaptureLostEvent& CaptureLostEvent)
+{
+	if (!IsScribbleEnabled())
+	{
+		return;
+	}
+	if (HasActiveNodeType())
+	{
+		if (bIsScribbleDragging.Get(false))
+		{
+			CancelNewNode();
+		}
+	}
+	SGraphPanel::OnMouseCaptureLost(CaptureLostEvent);
+}
+
+FReply SScribbleGraphPanel::OnMouseButtonUp( const FGeometry& MyGeometry, const FPointerEvent& MouseEvent )
+{
+	if (!IsScribbleEnabled())
+	{
+		return FReply::Unhandled();
+	}
+	if (HasActiveNodeType())
+	{
+		if (bIsScribbleDragging.Get(false))
+		{
+			UpdateNewNode(MyGeometry, MouseEvent);
+			CommitNewNode();
+			return FReply::Handled();
+		}
+	}
+	return SGraphPanel::OnMouseButtonUp(MyGeometry, MouseEvent);
+}
+
+FReply SScribbleGraphPanel::OnMouseMove( const FGeometry& MyGeometry, const FPointerEvent& MouseEvent )
+{
+	if (!IsScribbleEnabled())
+	{
+		return FReply::Unhandled();
+	}
+	if (HasActiveNodeType())
+	{
+		if (bIsScribbleDragging.Get(false))
+		{
+			UpdateNewNode(MyGeometry, MouseEvent);
+			return FReply::Handled();
+		}
+	}
+	return SGraphPanel::OnMouseMove(MyGeometry, MouseEvent);
+}
+
+void SScribbleGraphPanel::OnMouseLeave( const FPointerEvent& MouseEvent )
+{
+	if (!IsScribbleEnabled())
+	{
+		return;
+	}
+	SGraphPanel::OnMouseLeave(MouseEvent);
+}
+
+FReply SScribbleGraphPanel::OnKeyDown(const FGeometry& MyGeometry, const FKeyEvent& InKeyEvent)
+{
+	if (!IsScribbleEnabled())
+	{
+		return FReply::Unhandled();
+	}
+	if (CommandList.IsValid() && CommandList->ProcessCommandBindings(InKeyEvent))
+	{
+		return FReply::Handled();
+	}
+	return SGraphPanel::OnKeyDown(MyGeometry, InKeyEvent);
+}
+
+void SScribbleGraphPanel::Tick(const FGeometry& AllottedGeometry, const double InCurrentTime, const float InDeltaTime)
+{
+	SGraphPanel::Tick(AllottedGeometry, InCurrentTime, InDeltaTime);
+	UpdateScribbleVisibility();
+
+	if (WeakPanelToSync.IsValid())
+	{
+		if (TSharedPtr<SGraphPanel> PanelToSync = WeakPanelToSync.Pin())
+		{
+			if (IsScribbleEnabled())
+			{
+				PanelToSync->SetViewOffset(ViewOffset);
+				PanelToSync->SetZoomLevel(ZoomLevel);
+			}
+			else
+			{
+				SetViewOffset(PanelToSync->GetViewOffset());
+				SetZoomLevel(PanelToSync->GetZoomLevel());
+			}
+		}
+	}
+
+	if (bAllowNavigation && !WeakPanelToSync.IsValid())
+	{
+		if (FScribbleGraphData* ScribbleGraph = GetScribbleGraph())
+		{
+			ScribbleGraph->SetView(GetViewOffset(), GetZoomAmount());
+		}
+	}
+
+	if (const FScribbleGraphData* Graph = GetScribbleGraph())
+	{
+		if (Graph->SupportsAnchors())
+		{
+			FArrangedChildren ArrangedChildren(EVisibility::Visible);
+			ArrangeChildNodes(AllottedGeometry, ArrangedChildren);
+
+			for(int32 ChildIndex = 0; ChildIndex < ArrangedChildren.Num(); ChildIndex++)
+			{
+				const FArrangedWidget& ArrangedWidget = ArrangedChildren[ChildIndex];
+				TSharedRef<SScribbleGraphNode> GraphNode = StaticCastSharedRef<SScribbleGraphNode, SWidget>(ArrangedWidget.Widget);
+				const FScribbleNode* ScribbleNode = GraphNode->GetScribbleNode();
+				if (!ScribbleNode)
+				{
+					continue;
+				}
+				if (ScribbleNode->GetAnchor().IsNone())
+				{
+					continue;
+				}
+				UScribbleEdGraphNode* EdGraphNode = GraphNode->GetScribbleEdGraphNode();
+				if (!EdGraphNode)
+				{
+					continue;
+				}
+				const FVector2f NodePosition = ScribbleNode->GetPosition(true);
+				EdGraphNode->NodePosX = NodePosition.X;
+				EdGraphNode->NodePosY = NodePosition.Y;
+			}
+		}
+	}
+}
+
+int32 SScribbleGraphPanel::OnPaint(const FPaintArgs& Args, const FGeometry& AllottedGeometry, const FSlateRect& MyCullingRect,
+	FSlateWindowElementList& OutDrawElements, int32 LayerId, const FWidgetStyle& InWidgetStyle, bool bParentEnabled) const
+{
+	int32 LastLayerId = LayerId;
+
+	if (ScribbleEnabled.Get(false))
+	{
+		// draw a dark overlay below our content to indicate that we are in scribble mode
+		if (!ShouldDrawBackground.Get(false))
+		{
+			static const FSlateBrush* DefaultBackground = FAppStyle::GetBrush(TEXT("Graph.Panel.SolidBackground"));
+			static constexpr FLinearColor DarkenColor = FLinearColor(1,1,1, 0.6);
+			FSlateDrawElement::MakeBox(OutDrawElements, ++LastLayerId, AllottedGeometry.ToPaintGeometry(), DefaultBackground, ESlateDrawEffect::None, DarkenColor);
+		}
+	}
+	
+	if (GetZoomAmount() > 0.3f)
+	{
+		if (const FScribbleGraphData* Graph = GetScribbleGraph())
+		{
+			if (Graph->SupportsAnchors())
+			{
+				// update the outline color
+				const FLinearColor LineColor = UScribbleEditorSettings::Get()->AnchorColor;
+				const FLinearColor FillColor = LineColor * FLinearColor(0.35f,0.35f,0.35f,0.35f);;
+				RoundedFrameBrush->OutlineSettings.Color = LineColor;
+				RoundedFrameBrush->TintColor = FillColor;
+					
+				FArrangedChildren ArrangedChildren(EVisibility::Visible);
+				ArrangeChildNodes(AllottedGeometry, ArrangedChildren);
+				
+				for(int32 ChildIndex = 0; ChildIndex < ArrangedChildren.Num(); ChildIndex++)
+				{
+					const FArrangedWidget& ArrangedWidget = ArrangedChildren[ChildIndex];
+					TSharedRef<SScribbleGraphNode> GraphNode = StaticCastSharedRef<SScribbleGraphNode, SWidget>(ArrangedWidget.Widget);
+					const FScribbleNode* ScribbleNode = GraphNode->GetScribbleNode();
+					if (!ScribbleNode)
+					{
+						continue;
+					}
+					if (ScribbleNode->GetAnchor().IsNone())
+					{
+						continue;
+					}
+
+					TOptional<FVector2f> AnchorPosition = Graph->ResolveAnchor(ScribbleNode->GetAnchor());
+					if (!AnchorPosition.IsSet())
+					{
+						continue;
+					}
+
+					const FVector2f FrameOffset = FVector2f(20,20);
+					const FVector2f ScribbleNodePositionWithOffset = ScribbleNode->GetPosition(true) - FrameOffset;
+					const FVector2f ScribbleNodeSizeWithOffset = ScribbleNode->GetSize() + FrameOffset * 2.f;
+					FVector2f ScribbleNodeCorner = ScribbleNodePositionWithOffset;
+
+					// draw a rounded box around the anchored content
+					const FVector2f ScribbleNodePositionInView = (ScribbleNodePositionWithOffset - ViewOffset) * GetZoomAmount();
+             		const FVector2f ScribbleNodeSizeInView = ScribbleNodeSizeWithOffset * GetZoomAmount();
+                    const FPaintGeometry PaintGeometry = AllottedGeometry.ToPaintGeometry(ScribbleNodeSizeInView, FSlateLayoutTransform(ScribbleNodePositionInView));
+                    FSlateDrawElement::MakeBox(OutDrawElements, ++LastLayerId, PaintGeometry, RoundedFrameBrush.Get(), ESlateDrawEffect::None, FillColor);
+
+					// determine the shortest path to the box
+					const TArray<FVector2f> ScribbleNodeCorners = {
+						ScribbleNodeCorner,
+						ScribbleNodeCorner + (ScribbleNodeSizeWithOffset) * FVector2f(1.f, 0.f),
+						ScribbleNodeCorner + (ScribbleNodeSizeWithOffset) * FVector2f(1.f, 1.f),
+						ScribbleNodeCorner + (ScribbleNodeSizeWithOffset) * FVector2f(0.f, 1.f),
+					};
+					const TArray<FVector2f> ScribbleCandidates = {
+						FMath::Lerp<FVector2f>(ScribbleNodeCorners[0], ScribbleNodeCorners[1], 0.1f),
+						FMath::Lerp<FVector2f>(ScribbleNodeCorners[0], ScribbleNodeCorners[1], 0.9f),
+						FMath::Lerp<FVector2f>(ScribbleNodeCorners[1], ScribbleNodeCorners[2], 0.1f),
+						FMath::Lerp<FVector2f>(ScribbleNodeCorners[1], ScribbleNodeCorners[2], 0.9f),
+						FMath::Lerp<FVector2f>(ScribbleNodeCorners[2], ScribbleNodeCorners[3], 0.1f),
+						FMath::Lerp<FVector2f>(ScribbleNodeCorners[2], ScribbleNodeCorners[3], 0.9f),
+						FMath::Lerp<FVector2f>(ScribbleNodeCorners[3], ScribbleNodeCorners[0], 0.1f),
+						FMath::Lerp<FVector2f>(ScribbleNodeCorners[3], ScribbleNodeCorners[0], 0.9f)
+					};
+
+					float ClosestDistance = FLT_MAX;
+					for (const FVector2f& ScribbleCandidate : ScribbleCandidates)
+					{
+						const float Distance = FVector2f::DistSquared(ScribbleCandidate, AnchorPosition.GetValue());
+						if (Distance < ClosestDistance)
+						{
+							ClosestDistance = Distance;
+							ScribbleNodeCorner = ScribbleCandidate;
+						}
+					}
+
+					// draw a connecting line to the box 
+					const TArray<FVector2f> AnchorVertices = {
+						(AnchorPosition.GetValue() - ViewOffset) * GetZoomAmount(),
+						(ScribbleNodeCorner - ViewOffset) * GetZoomAmount()
+					};
+					FSlateDrawElement::MakeLines(OutDrawElements, LastLayerId, AllottedGeometry.ToPaintGeometry(), AnchorVertices, ESlateDrawEffect::None, LineColor, false, RoundedFrameBrush->OutlineSettings.Width);
+				}
+			}
+		}
+	}
+
+	return SGraphPanel::OnPaint(Args, AllottedGeometry, MyCullingRect, OutDrawElements, LayerId, InWidgetStyle, bParentEnabled);
+}
+
+void SScribbleGraphPanel::PostUndo(bool bSuccess)
+{
+	if (UScribbleEdGraph* ScribbleEdGraph = GetScribbleEdGraph())
+	{
+		ScribbleEdGraph->RebuildGraph();
+	}
+}
+
+void SScribbleGraphPanel::PostRedo(bool bSuccess)
+{
+	PostUndo(bSuccess);
+}
+
+void SScribbleGraphPanel::BindCommands()
+{
+	// create new command
+	const FScribbleEdGraphCommands& Commands = FScribbleEdGraphCommands::Get();
+
+	CommandList = MakeShared<FUICommandList>();
+
+	CommandList->MapAction(Commands.FrameSelection,
+		FExecuteAction::CreateLambda([this] ()
+		{
+			ZoomToFit(!ScribbleGraphObj->GetSelectedNodeIds().IsEmpty());
+		}),
+		FCanExecuteAction::CreateLambda([this] () -> bool
+		{
+			return ScribbleGraphObj->GetNodes().Num() > 0;
+		})		
+	);
+
+	CommandList->MapAction(Commands.DeleteSelection,
+		FExecuteAction::CreateLambda([this] ()
+		{
+			const FScopedTransaction _(LOCTEXT("RemoveSelectedScribbleNodes", "Remove Selected Scribble Nodes"));
+			ScribbleGraphObj->RemoveSelectedNodes();
+		}),
+		FCanExecuteAction::CreateLambda([this] () -> bool
+		{
+			return ScribbleGraphObj->GetSelectedNodeIds().Num() > 0;
+		})		
+	);
+
+	CommandList->MapAction(Commands.SelectAll,
+		FExecuteAction::CreateLambda([this] ()
+		{
+			ScribbleGraphObj->SelectAllNodes();
+		}),
+		FCanExecuteAction::CreateLambda([this] () -> bool
+		{
+			return true;
+		})		
+	);
+
+	CommandList->MapAction(Commands.GroupSelection,
+		FExecuteAction::CreateLambda([this] ()
+		{
+			const FScopedTransaction _(LOCTEXT("GroupSelectedScribbleNodes", "Group Scribble Nodes"));
+			ScribbleGraphObj->GroupSelectedNodes();
+		}),
+		FCanExecuteAction::CreateLambda([this] () -> bool
+		{
+			return ScribbleGraphObj->GetSelectedNodeIds().Num() > 1;
+		})		
+	);
+
+	CommandList->MapAction(Commands.UngroupSelection,
+		FExecuteAction::CreateLambda([this] ()
+		{
+			const FScopedTransaction _(LOCTEXT("UngroupSelectedScribbleNode", "Ungroup Scribble Node"));
+			ScribbleGraphObj->UngroupSelectedNodes();
+		}),
+		FCanExecuteAction::CreateLambda([this] () -> bool
+		{
+			return ScribbleGraphObj->GetSelectedNodeIds().Num() == 1;
+		})		
+	);
+
+	CommandList->MapAction(Commands.SelectionTool,
+		FExecuteAction::CreateLambda([this] ()
+		{
+			ActivateSelectionMode();
+		})
+	);
+
+	CommandList->MapAction(Commands.BrushTool,
+		FExecuteAction::CreateLambda([this] ()
+		{
+			SetActiveNodeType(EScribbleNodeType::LineStrip);
+		})
+	);
+}
+
+void SScribbleGraphPanel::UpdateScribbleVisibility()
+{
+	if (GetVisibility() == EVisibility::Visible ||
+		GetVisibility() == EVisibility::HitTestInvisible)
+	{
+		SetVisibility(IsScribbleEnabled() ? EVisibility::Visible : EVisibility::HitTestInvisible);
+	}
+}
+
+void SScribbleGraphPanel::HandleSelectionChanged(const FGraphPanelSelectionSet& SelectionSet)
+{
+	UScribbleEdGraph* ScribbleEdGraph = GetScribbleEdGraph();
+	if (!ScribbleEdGraph)
+	{
+		return;
+	}
+
+	if (ScribbleEdGraph->bIsSelecting)
+	{
+		return;
+	}
+	const TGuardValue<bool> _(ScribbleEdGraph->bIsSelecting, true);
+
+	TArray<FGuid> NodeIds;
+	for (const UObject* SelectedObject : SelectionSet)
+	{
+		if (const UScribbleEdGraphNode* Node = Cast<UScribbleEdGraphNode>(SelectedObject))
+		{
+			NodeIds.Add(Node->GetNodeId());
+		}
+	}
+	ScribbleEdGraph->SelectNodes(NodeIds);
+}
+
+bool SScribbleGraphPanel::UpdateNewNode(const FGeometry& MyGeometry, const FPointerEvent& MouseEvent)
+{
+	if (!bIsScribbleDragging.Get(false))
+	{
+		return false;
+	}
+	
+	if (!NewNode)
+	{
+		Transaction = MakeShared<FScopedTransaction>(LOCTEXT("AddScribbleNode", "Add Scribble Node"));
+		if (FScribbleGraphData* ScribbleGraph = GetScribbleGraph())
+		{
+			ScribbleGraph->Modify();
+		}
+		switch (ScribbleNodeType.Get(EScribbleNodeType::LineStrip))
+		{
+			case EScribbleNodeType::LineStrip:
+			{
+				TSharedPtr<FLineStripScribbleNode> LineStripScribbleNode = MakeShared<FLineStripScribbleNode>();
+				LineStripScribbleNode->LineStrips.Emplace();
+				LineStripScribbleNode->LineStrips.Last().Thickness = ScribbleThickness.Get(4.f);
+				LineStripScribbleNode->LineStrips.Last().Color = ScribbleColor.Get(FLinearColor::Blue);
+				NewNode = LineStripScribbleNode;
+				break;
+			}
+			default:
+			{
+				return false;
+			}
+		}
+
+		NewNode->SetPosition(GetViewOffset());
+		ScribbleGraphObj->AddScribbleNode(NewNode);
+		ScribbleGraphObj->ClearSelection();
+	}
+
+
+	const FVector2f MousePosition = MouseEvent.GetScreenSpacePosition();
+	const FVector2f GraphEditorLocation = GetViewOffset();
+	const float GraphEditorZoomAmount = GetZoomAmount();
+	const FVector2f WidgetPosition = MyGeometry.AbsoluteToLocal(MousePosition) / GraphEditorZoomAmount + GraphEditorLocation;
+	
+	if (!NewNode->AppendMouseEvent(MyGeometry, MouseEvent, WidgetPosition))
+	{
+		return false;
+	}
+	if (UScribbleEdGraphNode* EdGraphNode = ScribbleGraphObj->FindNode(NewNode->GetId()))
+	{
+		EdGraphNode->SetPosition(NewNode->GetPosition());
+		EdGraphNode->SetSize(NewNode->GetSize());
+	}
+	return true;
+}
+
+bool SScribbleGraphPanel::CancelNewNode()
+{
+	bIsScribbleDragging.Reset();
+	if (!NewNode)
+	{
+		return false;
+	}
+	ScribbleGraphObj->RemoveScribbleNode(NewNode);
+	if (Transaction)
+	{
+		Transaction->Cancel();
+	}
+	Transaction.Reset();
+	NewNode.Reset();
+	return true;
+}
+
+bool SScribbleGraphPanel::CommitNewNode()
+{
+	bIsScribbleDragging.Reset();
+	if (!NewNode)
+	{
+		return false;
+	}
+
+	if (!ScribbleGraphObj)
+	{
+		NewNode.Reset();
+		return true;
+	}
+
+	const FBox2f Bounds = NewNode->GetContentBounds();
+	const FVector2f TopLeft = NewNode->GetPosition() + Bounds.GetCenter() - Bounds.GetExtent();
+
+	NewNode->Modify();
+	NewNode->IncrementChangeBracket();
+	if (!NewNode->DownSample(ScribblePrecision.Get(3.f)))
+	{
+		NewNode->DecrementChangeBracket();
+		return CancelNewNode();
+	}
+	
+	NewNode->OffsetPosition(TopLeft);
+	NewNode->DecrementChangeBracket();
+	NewNode.Reset();
+	Transaction.Reset();
+
+	return true;
+}
+
+SScribbleGraph::~SScribbleGraph()
+{
+}
+
+void SScribbleGraph::Construct(const FArguments& InArgs)
+{
+	bIsScribbleEnabled = false;
+	
+	ChildSlot
+	[
+		SNew(SOverlay)
+
+		+SOverlay::Slot()
+		.HAlign(HAlign_Fill)
+		.VAlign(VAlign_Fill)
+		.Padding(0)
+		[
+			SAssignNew(GraphPanel, SScribbleGraphPanel)
+			.GraphData(InArgs._GraphData)
+			.ScribbleEnabled(this, &SScribbleGraph::IsScribbleEnabled)
+			.ScribbleColor(this, &SScribbleGraph::GetScribbleColor)
+			.ScribbleThickness(this, &SScribbleGraph::GetScribbleThickness)
+			.ScribblePrecision(this, &SScribbleGraph::GetScribblePrecision)
+			.ShouldDrawBackground(InArgs._ShouldDrawBackground)
+			.ShouldDrawSurroundingShadow(InArgs._ShouldDrawSurroundingShadow)
+			.AllowNavigation(InArgs._AllowNavigation)
+			.PanelToSync(InArgs._PanelToSync)
+		]
+
+		+SOverlay::Slot()
+		.HAlign(InArgs._ToolbarHAlignment)
+		.VAlign(InArgs._ToolbarVAlignment)
+		.Padding(InArgs._ToolbarPadding)
+		[
+			SNew(SHorizontalBox)
+
+			+SHorizontalBox::Slot()
+			.AutoWidth()
+			.HAlign(HAlign_Left)
+			.VAlign(VAlign_Center)
+			.Padding(0)
+			[
+				SNew(SButton)
+				.ButtonStyle( FAppStyle::Get(), "SimpleButton" )
+				.ToolTipText(LOCTEXT("EnabledTooltip", "Enabled"))
+				.OnClicked_Lambda([this]() -> FReply
+				{
+					bIsScribbleEnabled = !bIsScribbleEnabled;
+					return FReply::Handled();
+				})
+				.ContentPadding(2)
+				[
+					SNew(SImage)
+					.ColorAndOpacity_Lambda([this]()
+					{
+						if (bIsScribbleEnabled)
+						{
+							return FStyleColors::AccentBlue;
+						}
+						return FSlateColor::UseForeground();
+					})
+					.DesiredSizeOverride(FVector2D(20,20))
+					.Image(FAppStyle::Get().GetBrush("MeshPaint.Brush"))
+				]
+			]
+			
+			+SHorizontalBox::Slot()
+			.AutoWidth()
+			.HAlign(HAlign_Left)
+			.VAlign(VAlign_Center)
+			.Padding(8, 0, 0, 0)
+			[
+				SNew(SButton)
+				.Visibility_Lambda([this]()
+				{
+					return bIsScribbleEnabled ? EVisibility::Visible : EVisibility::Collapsed;
+				})
+				.ButtonStyle( FAppStyle::Get(), "SimpleButton" )
+				.ToolTipText(LOCTEXT("LineStripTooltip", "Draw Lines"))
+				.OnClicked_Lambda([this]() -> FReply
+				{
+					if (GraphPanel->GetActiveNodeType() == EScribbleNodeType::Invalid)
+					{
+						GraphPanel->SetActiveNodeType(EScribbleNodeType::LineStrip);
+					}
+					else
+					{
+						GraphPanel->ActivateSelectionMode();
+					}
+					return FReply::Handled();
+				})
+				.ContentPadding(2)
+				[
+					SNew(SImage)
+					.ColorAndOpacity_Lambda([this]()
+					{
+						if (GraphPanel->GetActiveNodeType() == EScribbleNodeType::Invalid)
+						{
+							return FStyleColors::AccentBlue;
+						}
+						return FSlateColor::UseForeground();
+					})
+					.DesiredSizeOverride(FVector2D(20,20))
+					.Image(FAppStyle::Get().GetBrush("MeshPaint.SelectVertex"))
+				]
+			]
+
+			+SHorizontalBox::Slot()
+			.AutoWidth()
+			.HAlign(HAlign_Left)
+			.VAlign(VAlign_Center)
+			.Padding(4,0,0,0)
+			[
+				SNew(SColorBlock)
+				.Visibility_Lambda([this]()
+				{
+					return bIsScribbleEnabled ? EVisibility::Visible : EVisibility::Collapsed;
+				})
+				.CornerRadius(FVector4(3.f, 3.f, 3.f, 3.f))
+				.Color_Lambda([]() -> FLinearColor
+				{
+					return UScribbleEditorSettings::Get()->Color;
+				})
+				.OnMouseButtonDown_Lambda([this](const FGeometry&, const FPointerEvent&) -> FReply
+				{
+					FColorPickerArgs PickerArgs;
+					PickerArgs.bIsModal = true;
+					PickerArgs.ParentWidget = AsShared();
+					PickerArgs.InitialColor = UScribbleEditorSettings::Get()->Color;
+					PickerArgs.OnColorCommitted = FOnLinearColorValueChanged::CreateLambda([](FLinearColor NewValue)
+					{
+						UScribbleEditorSettings::Get()->Color = NewValue;
+						UScribbleEditorSettings::Get()->SaveConfig();
+					});
+
+					OpenColorPicker(PickerArgs);
+					return FReply::Handled();
+				})
+			]
+
+			+SHorizontalBox::Slot()
+			.AutoWidth()
+			.HAlign(HAlign_Left)
+			.VAlign(VAlign_Center)
+			.Padding(0)
+			[
+				SNew(SButton)
+				.Visibility_Lambda([this]()
+				{
+					if (!bIsScribbleEnabled)
+					{
+						return EVisibility::Collapsed;
+					}
+					if (FScribbleGraphData* ScribbleGraph = GraphPanel->GetScribbleGraph())
+					{
+						if (!ScribbleGraph->SupportsAnchors())
+						{
+							return EVisibility::Collapsed;
+						}
+					}
+					return EVisibility::Visible;
+				})
+				.IsEnabled_Lambda([this]
+				{
+					if (FScribbleGraphData* ScribbleGraph = GraphPanel->GetScribbleGraph())
+					{
+						if (!ScribbleGraph->GetCurrentAnchor().IsNone())
+						{
+							return true;
+						}
+					}
+					return false;
+				})
+				.ButtonStyle( FAppStyle::Get(), "SimpleButton" )
+				.ToolTipText(LOCTEXT("AnchorTooltip", "Anchor to selected node"))
+				.OnClicked_Lambda([this]() -> FReply
+				{
+					if (UScribbleEdGraph* ScribbleEdGraph = GraphPanel->GetScribbleEdGraph())
+					{
+						const TArray<FGuid> SelectedIds = ScribbleEdGraph->GetSelectedNodeIds();
+						if (!SelectedIds.IsEmpty())
+						{
+							if (FScribbleGraphData* ScribbleGraph = GraphPanel->GetScribbleGraph())
+							{
+								const FName& CurrentAnchor = ScribbleGraph->GetCurrentAnchor();
+
+								bool bShouldEnableAnchor = true;
+								for (const FGuid& NodeId : SelectedIds)
+								{
+									if (FScribbleNode* ScribbleNode = ScribbleGraph->FindNode(NodeId))
+									{
+										if (!ScribbleNode->GetAnchor().IsNone())
+										{
+											bShouldEnableAnchor = false;
+											break;
+										}
+									}
+								}
+
+								ScribbleGraph->Modify();
+								for (const FGuid& NodeId : SelectedIds)
+								{
+									if (FScribbleNode* ScribbleNode = ScribbleGraph->FindNode(NodeId))
+									{
+										ScribbleNode->SetAnchor(bShouldEnableAnchor ? CurrentAnchor : NAME_None);
+									}
+								}
+							}
+						}
+					}
+						
+					return FReply::Handled();
+				})
+				.ContentPadding(2)
+				[
+					SNew(SImage)
+					.ColorAndOpacity_Lambda([this]()
+					{
+						if (const UScribbleEdGraph* ScribbleEdGraph = GraphPanel->GetScribbleEdGraph())
+						{
+							for (const FGuid& Guid : GraphPanel->GetScribbleEdGraph()->GetSelectedNodeIds())
+							{
+								if (const UScribbleEdGraphNode* ScribbleEdGraphNode = ScribbleEdGraph->FindNode(Guid))
+								{
+									if (!ScribbleEdGraphNode->GetScribbleNode()->GetAnchor().IsNone())
+									{
+										return FStyleColors::AccentBlue;
+									}
+								}
+							}
+						}
+						return FSlateColor::UseForeground();
+					})
+					.DesiredSizeOverride(FVector2D(20,20))
+					.Image(FAppStyle::Get().GetBrush("Icons.Snap"))
+				]
+			]
+
+			+SHorizontalBox::Slot()
+			.AutoWidth()
+			.HAlign(HAlign_Left)
+			.VAlign(VAlign_Center)
+			.Padding(4,0,0,0)
+			[
+				SNew(SSpinBox<float>)
+				.ToolTipText(LOCTEXT("ThicknessTooltip", "Thickness of the brush"))
+				.Delta(0.25f)
+				.MinSliderValue(3.f)
+				.MaxSliderValue(24.f)
+				.MinFractionalDigits(2)
+				.MaxFractionalDigits(2)
+				.PreventThrottling(true)
+				.Visibility_Lambda([this]()
+				{
+					return bIsScribbleEnabled ? EVisibility::Visible : EVisibility::Collapsed;
+				})
+				.Value_Lambda([]() -> float
+				{
+					return UScribbleEditorSettings::Get()->Thickness;
+				})
+				.OnValueChanged_Lambda([](float InNewValue)
+				{
+					UScribbleEditorSettings::Get()->Thickness = InNewValue;
+				})
+				.OnValueCommitted_Lambda([](float InNewValue, ETextCommit::Type InCommitType)
+				{
+					UScribbleEditorSettings::Get()->Thickness = InNewValue;
+					UScribbleEditorSettings::Get()->SaveConfig();
+				})
+			]
+
+			+SHorizontalBox::Slot()
+			.AutoWidth()
+			.HAlign(HAlign_Left)
+			.VAlign(VAlign_Center)
+			.Padding(4,0,0,0)
+			[
+				SNew(SSpinBox<float>)
+				.ToolTipText(LOCTEXT("SmoothingTooltip", "The smoothing of the strokes\n0 == maintain curvature\n1 == smooth / straigthen lines"))
+				.Delta(0.025f)
+				.MinSliderValue(0.0)
+				.MaxSliderValue(1.0)
+				.MinFractionalDigits(2)
+				.MaxFractionalDigits(2)
+				.PreventThrottling(true)
+				.Visibility_Lambda([this]()
+				{
+					return bIsScribbleEnabled ? EVisibility::Visible : EVisibility::Collapsed;
+				})
+				.Value_Lambda([]() -> float
+				{
+					return UScribbleEditorSettings::Get()->GetSmoothing();
+				})
+				.OnValueChanged_Lambda([](float InNewValue)
+				{
+					UScribbleEditorSettings::Get()->SetSmoothing(InNewValue);
+				})
+				.OnValueCommitted_Lambda([](float InNewValue, ETextCommit::Type InCommitType)
+				{
+					UScribbleEditorSettings::Get()->SetSmoothing(InNewValue);
+					UScribbleEditorSettings::Get()->SaveConfig();
+				})
+			]
+
+			+SHorizontalBox::Slot()
+			.AutoWidth()
+			.HAlign(HAlign_Left)
+			.VAlign(VAlign_Center)
+			.Padding(16,0,0,0)
+			[
+				SNew(SButton)
+				.Visibility_Lambda([this]()
+				{
+					if (!bIsScribbleEnabled)
+					{
+						return EVisibility::Collapsed;
+					}
+					return EVisibility::Visible;
+				})
+				.IsEnabled_Lambda([this]
+				{
+					return GUnrealEd->Trans->CanUndo();
+				})
+				.ButtonStyle( FAppStyle::Get(), "SimpleButton" )
+				.ToolTipText(LOCTEXT("Undo", "Undo"))
+				.OnClicked_Lambda([this]() -> FReply
+				{
+					GEditor->UndoTransaction();
+					return FReply::Handled();
+				})
+				.ContentPadding(2)
+				[
+					SNew(SImage)
+					.ColorAndOpacity(FSlateColor::UseForeground())
+					.Image(FAppStyle::Get().GetBrush("GenericCommands.Undo"))
+				]
+			]
+
+			+SHorizontalBox::Slot()
+			.AutoWidth()
+			.HAlign(HAlign_Left)
+			.VAlign(VAlign_Center)
+			.Padding(0)
+			[
+				SNew(SButton)
+				.Visibility_Lambda([this]()
+				{
+					if (!bIsScribbleEnabled)
+					{
+						return EVisibility::Collapsed;
+					}
+					return EVisibility::Visible;
+				})
+				.IsEnabled_Lambda([this]
+				{
+					return GUnrealEd->Trans->CanRedo();
+				})
+				.ButtonStyle( FAppStyle::Get(), "SimpleButton" )
+				.ToolTipText(LOCTEXT("Redo", "Redo"))
+				.OnClicked_Lambda([this]() -> FReply
+				{
+					GEditor->RedoTransaction();
+					return FReply::Handled();
+				})
+				.ContentPadding(2)
+				[
+					SNew(SImage)
+					.ColorAndOpacity(FSlateColor::UseForeground())
+					.Image(FAppStyle::Get().GetBrush("GenericCommands.Redo"))
+				]
+			]
+			
+			+SHorizontalBox::Slot()
+			.AutoWidth()
+			.HAlign(HAlign_Left)
+			.VAlign(VAlign_Center)
+			.Padding(0)
+			[
+				SNew(SButton)
+				.Visibility_Lambda([this]()
+				{
+					if (!bIsScribbleEnabled)
+					{
+						return EVisibility::Collapsed;
+					}
+					return EVisibility::Visible;
+				})
+				.IsEnabled_Lambda([this]
+				{
+					return !GraphPanel->GetScribbleEdGraph()->GetSelectedNodeIds().IsEmpty();
+				})
+				.ButtonStyle( FAppStyle::Get(), "SimpleButton" )
+				.ToolTipText(LOCTEXT("Delete", "Delete"))
+				.OnClicked_Lambda([this]() -> FReply
+				{
+					GraphPanel->GetScribbleEdGraph()->RemoveSelectedNodes();
+					return FReply::Handled();
+				})
+				.ContentPadding(2)
+				[
+					SNew(SImage)
+					.ColorAndOpacity(FSlateColor::UseForeground())
+					.Image(FAppStyle::Get().GetBrush("Icons.Delete"))
+				]
+			]
+		]
+	];
+
+	SetVisibility(EVisibility::SelfHitTestInvisible);
+}
+
+FLinearColor SScribbleGraph::GetScribbleColor() const
+{
+	return UScribbleEditorSettings::Get()->Color;
+}
+
+float SScribbleGraph::GetScribbleThickness() const
+{
+	return UScribbleEditorSettings::Get()->Thickness;
+}
+
+float SScribbleGraph::GetScribblePrecision() const
+{
+	return UScribbleEditorSettings::Get()->Precision;
+}
+
+#undef LOCTEXT_NAMESPACE
diff --git a/Engine/Plugins/Experimental/Animation/Scribble/Source/ScribbleEditor/Private/SScribbleGraphNode.cpp b/Engine/Plugins/Experimental/Animation/Scribble/Source/ScribbleEditor/Private/SScribbleGraphNode.cpp
new file mode 100644
--- /dev/null
+++ b/Engine/Plugins/Experimental/Animation/Scribble/Source/ScribbleEditor/Private/SScribbleGraphNode.cpp
@@ -0,0 +1,184 @@
+// Copyright Epic Games, Inc. All Rights Reserved.
+
+#include "SScribbleGraphNode.h"
+#include "SScribbleGraph.h"
+
+#include "Delegates/Delegate.h"
+#include "EdGraph/EdGraphNode.h"
+#include "EdGraph/EdGraphPin.h"
+#include "GenericPlatform/ICursor.h"
+#include "HAL/Platform.h"
+#include "HAL/PlatformCrt.h"
+#include "Math/Vector2D.h"
+#include "Misc/AssertionMacros.h"
+#include "Misc/Attribute.h"
+#include "ScribbleEdGraphNode.h"
+#include "SGraphPin.h"
+#include "SNodePanel.h"
+#include "Styling/AppStyle.h"
+#include "Widgets/Images/SImage.h"
+#include "Widgets/Layout/SBorder.h"
+#include "Widgets/Layout/SSpacer.h"
+#include "ScribbleEdGraph.h"
+#include "SGraphPanel.h"
+
+class SWidget;
+class UEdGraphSchema;
+struct FSlateBrush;
+
+#define LOCTEXT_NAMESPACE "RigDependencyGraphEditor"
+
+void SScribbleGraphNode::Construct(const FArguments& InArgs, class UScribbleEdGraphNode* InNode)
+{
+	GraphNode = InNode;
+	UpdateGraphNode();
+}
+
+const FSlateBrush* SScribbleGraphNode::GetShadowBrush(bool bSelected) const
+{
+	return nullptr;
+}
+
+void SScribbleGraphNode::UpdateGraphNode()
+{
+	this->ContentScale.Bind( this, &SGraphNode::GetContentScale );
+}
+
+FReply SScribbleGraphNode::OnMouseButtonDown(const FGeometry& MyGeometry, const FPointerEvent& MouseEvent)
+{
+	if (const FScribbleNode* ScribbleNode = GetScribbleNode())
+	{
+		const FVector2f LocalPosition = MyGeometry.AbsoluteToLocal(MouseEvent.GetScreenSpacePosition());
+		if (!ScribbleNode->IntersectsCursorPosition(LocalPosition))
+		{
+		}
+	}
+	return SGraphNode::OnMouseButtonDown(MyGeometry, MouseEvent);
+}
+
+FReply SScribbleGraphNode::OnMouseMove(const FGeometry& MyGeometry, const FPointerEvent& MouseEvent)
+{
+	return SGraphNode::OnMouseMove(MyGeometry, MouseEvent);
+}
+
+FReply SScribbleGraphNode::OnMouseButtonUp(const FGeometry& MyGeometry, const FPointerEvent& MouseEvent)
+{
+	return SGraphNode::OnMouseButtonUp(MyGeometry, MouseEvent);
+}
+
+void SScribbleGraphNode::MoveTo(const FVector2f& NewPosition, FNodeSet& NodeFilter, bool bMarkDirty)
+{
+	if (FScribbleNode* ScribbleNode = GetScribbleNode())
+	{
+		ScribbleNode->SetPosition(NewPosition);
+	}
+}
+
+int32 SScribbleGraphNode::OnPaint(const FPaintArgs& Args, const FGeometry& AllottedGeometry, const FSlateRect& MyCullingRect,
+                                  FSlateWindowElementList& OutDrawElements, int32 LayerId, const FWidgetStyle& InWidgetStyle, bool bParentEnabled) const
+{
+	int32 LastLayerId = SGraphNode::OnPaint(Args, AllottedGeometry, MyCullingRect, OutDrawElements, LayerId, InWidgetStyle, bParentEnabled);
+	if (const UScribbleEdGraphNode* ScribbleEdGraphNode = GetScribbleEdGraphNode())
+	{
+		if (const FScribbleNode* ScribbleNode = ScribbleEdGraphNode->GetScribbleNode())
+		{
+			bool bSelected = false;
+			if (const UScribbleEdGraph* ScribbleEdGraph = Cast<UScribbleEdGraph>(ScribbleEdGraphNode->GetGraph()))
+			{
+				bSelected = ScribbleEdGraph->IsNodeSelected(ScribbleEdGraphNode);
+			}
+
+			float ZoomAmount = 1.f;
+			if (OwnerGraphPanelPtr.IsValid())
+			{
+				if (TSharedPtr<SScribbleGraphPanel> GraphPanel = StaticCastSharedPtr<SScribbleGraphPanel,SGraphPanel>(OwnerGraphPanelPtr.Pin()))
+				{
+					ZoomAmount = GraphPanel->GetZoomAmount();
+
+					// don't render any node as selected if the
+					// panel is not enabled for scribbling
+					if (!GraphPanel->IsScribbleEnabled())
+					{
+						bSelected = false;
+					}
+				}
+			}
+			LastLayerId = ScribbleNode->OnPaint(AllottedGeometry, OutDrawElements, ++LastLayerId, InWidgetStyle, bSelected, ZoomAmount);
+		}
+	}
+	return LastLayerId;
+}
+
+FVector2D SScribbleGraphNode::ComputeDesiredSize(float LayoutScaleMultiplier) const
+{
+	return FVector2D(GetSizeFromScribbleNode());
+}
+
+bool SScribbleGraphNode::CanBeSelected(const FVector2f& MousePositionInNode) const
+{
+	if (const FScribbleNode* ScribbleNode = GetScribbleNode())
+	{
+		if (!ScribbleNode->IntersectsCursorPosition(MousePositionInNode))
+		{
+			return false;
+		}
+	}
+	return true;
+}
+
+const UScribbleEdGraphNode* SScribbleGraphNode::GetScribbleEdGraphNode() const
+{
+	return Cast<UScribbleEdGraphNode>(GraphNode);
+}
+
+UScribbleEdGraphNode* SScribbleGraphNode::GetScribbleEdGraphNode()
+{
+	return Cast<UScribbleEdGraphNode>(GraphNode);
+}
+
+const FScribbleNode* SScribbleGraphNode::GetScribbleNode() const
+{
+	if (const UScribbleEdGraphNode* EdGraphNode = GetScribbleEdGraphNode())
+	{
+		return EdGraphNode->GetScribbleNode();
+	}
+	return nullptr;
+}
+
+FScribbleNode* SScribbleGraphNode::GetScribbleNode()
+{
+	if (UScribbleEdGraphNode* EdGraphNode = GetScribbleEdGraphNode())
+	{
+		return EdGraphNode->GetScribbleNode();
+	}
+	return nullptr;
+}
+
+const FScribbleGraphData* SScribbleGraphNode::GetScribbleGraphData() const
+{
+	if (const UScribbleEdGraphNode* EdGraphNode = GetScribbleEdGraphNode())
+	{
+		return EdGraphNode->GetScribbleGraphData();
+	}
+	return nullptr;
+}
+
+FScribbleGraphData* SScribbleGraphNode::GetScribbleGraphData()
+{
+	if (UScribbleEdGraphNode* EdGraphNode = GetScribbleEdGraphNode())
+	{
+		return EdGraphNode->GetScribbleGraphData();
+	}
+	return nullptr;
+}
+
+FVector2f SScribbleGraphNode::GetSizeFromScribbleNode() const
+{
+	if (const UScribbleEdGraphNode* ScribbleEdGraphNode = GetScribbleEdGraphNode())
+	{
+		return ScribbleEdGraphNode->GetSize();
+	}
+	return FVector2f();
+}
+
+#undef LOCTEXT_NAMESPACE
diff --git a/Engine/Plugins/Experimental/Animation/Scribble/Source/ScribbleEditor/Private/ScribbleEdGraph.cpp b/Engine/Plugins/Experimental/Animation/Scribble/Source/ScribbleEditor/Private/ScribbleEdGraph.cpp
new file mode 100644
--- /dev/null
+++ b/Engine/Plugins/Experimental/Animation/Scribble/Source/ScribbleEditor/Private/ScribbleEdGraph.cpp
@@ -0,0 +1,402 @@
+// Copyright Epic Games, Inc. All Rights Reserved.
+
+#include "ScribbleEdGraph.h"
+
+#include "ITransportControl.h"
+#include "ScribbleEdGraphNode.h"
+#include "ScribbleEdGraphPanelNodeFactory.h"
+#include "ScribbleEdGraphSchema.h"
+
+#include UE_INLINE_GENERATED_CPP_BY_NAME(ScribbleEdGraph)
+
+UScribbleEdGraph::UScribbleEdGraph(const FObjectInitializer& ObjectInitializer)
+	: Super(ObjectInitializer)
+	, bIsSelecting(false)
+{
+}
+
+void UScribbleEdGraph::Initialize(const TSharedPtr<FScribbleGraphData>& InScribbleGraphData)
+{
+	if (!InScribbleGraphData)
+	{
+		return;
+	}
+
+	GraphData = InScribbleGraphData;
+
+	for (TSharedPtr<FScribbleNode> ScribbleNode : *InScribbleGraphData)
+	{
+		check(ScribbleNode);
+		AddScribbleNode(ScribbleNode);
+	}
+}
+
+const UScribbleEdGraphNode* UScribbleEdGraph::FindNode(const FGuid& InNodeId) const
+{
+	if (UScribbleEdGraphNode* const* NodePtr = NodeIdLookup.Find(InNodeId))
+	{
+		return *NodePtr;
+	}
+	return nullptr;
+}
+
+UScribbleEdGraphNode* UScribbleEdGraph::FindNode(const FGuid& InNodeId)
+{
+	const UScribbleEdGraph* ConstThis = this;
+	return const_cast<UScribbleEdGraphNode*>(ConstThis->FindNode(InNodeId));
+}
+
+FGuid UScribbleEdGraph::AddScribbleNode(const TSharedPtr<FScribbleNode>& InNode)
+{
+	if (!GraphData || !InNode)
+	{
+		return FGuid();
+	}
+
+	if (const UScribbleEdGraphNode* ExistingNode = FindNode(InNode->GetId()))
+	{
+		return ExistingNode->GetNodeId();
+	}
+
+	(void)GraphData->AddNode(InNode);
+	
+	UScribbleEdGraphNode* ScribbleEdGraphNode = NewObject<UScribbleEdGraphNode>(this);
+	ScribbleEdGraphNode->Initialize(InNode);
+	NodeIdLookup.Add(InNode->GetId(), ScribbleEdGraphNode);
+
+	AddNode(ScribbleEdGraphNode, false, false);
+	return InNode->GetId();
+}
+
+bool UScribbleEdGraph::RemoveScribbleNode(const TSharedPtr<FScribbleNode>& InNode)
+{
+	if (!GraphData || !InNode)
+	{
+		return false;
+	}
+
+	if (!GraphData->RemoveNode(InNode))
+	{
+		return false;
+	}
+	
+	UScribbleEdGraphNode* ScribbleEdGraphNode = FindNode(InNode->GetId());
+	NodeIdLookup.Remove(InNode->GetId());
+	SelectedNodes.Remove(InNode->GetId());
+	if (ScribbleEdGraphNode)
+	{
+		RemoveNode(ScribbleEdGraphNode);
+	}
+	
+	return true;
+}
+
+bool UScribbleEdGraph::RemoveScribbleNode(const FGuid& InGuid)
+{
+	if (!GraphData)
+	{
+		return false;
+	}
+	TSharedPtr<FScribbleNode> Node = GraphData->FindNodePtr(InGuid);
+	if (!Node)
+	{
+		return false;
+	}
+	return RemoveScribbleNode(Node);
+}
+
+bool UScribbleEdGraph::Contains(const FGuid& InNodeId) const
+{
+	return FindNode(InNodeId) != nullptr;
+}
+
+FScribbleGraphData* UScribbleEdGraph::GetGraphData()
+{
+	return GraphData.Get();
+}
+
+const FScribbleGraphData* UScribbleEdGraph::GetGraphData() const
+{
+	return GraphData.Get();
+}
+
+void UScribbleEdGraph::RebuildGraph()
+{
+	if (!GraphData)
+	{
+		return;
+	}
+
+	// we may want to be lazy here and only remove / add what's necessary
+	Nodes.Reset();
+	NodeIdLookup.Reset();
+
+	GraphData->IncrementChangeBracket();
+	
+	for (const TSharedPtr<FScribbleNode>& ScribbleNode : *GraphData)
+	{
+		check(ScribbleNode);
+		AddScribbleNode(ScribbleNode);
+	}
+
+	GraphData->DecrementChangeBracket();
+	NotifyGraphChanged();
+}
+
+void UScribbleEdGraph::SelectAllNodes()
+{
+	SelectNodes(GetAllNodeIds());
+}
+
+void UScribbleEdGraph::ClearSelection()
+{
+	SelectNodes({});
+}
+
+void UScribbleEdGraph::RemoveAllNodes()
+{
+	if (Nodes.IsEmpty())
+	{
+		return;
+	}
+
+	if (!GraphData)
+	{
+		return;
+	}
+
+	GraphData->IncrementChangeBracket();
+	
+	TArray<UEdGraphNode*> NodesToRemove = Nodes;
+	for (int32 NodeIndex = 0; NodeIndex < NodesToRemove.Num(); ++NodeIndex)
+	{
+		if (const UScribbleEdGraphNode* NodeToRemove = Cast<UScribbleEdGraphNode>(NodesToRemove[NodeIndex]))
+		{
+			RemoveScribbleNode(NodeToRemove->GetNodeId());
+		}
+	}
+
+	Nodes.Reset();
+	
+	GraphData->DecrementChangeBracket();
+	
+	NotifyGraphChanged();
+}
+
+TArray<FGuid> UScribbleEdGraph::GetAllNodeIds() const
+{
+	TArray<FGuid> NodeIds;
+	NodeIds.Reserve(Nodes.Num());
+	for (const TObjectPtr<UEdGraphNode>& Node : Nodes)
+	{
+		if (const UScribbleEdGraphNode* ScribbleEdGraphNode = Cast<UScribbleEdGraphNode>(Node))
+		{
+			NodeIds.Add(ScribbleEdGraphNode->GetNodeId());
+		}
+	}
+	return NodeIds;
+}
+
+void UScribbleEdGraph::RemoveSelectedNodes()
+{
+	if (Nodes.IsEmpty() || SelectedNodes.IsEmpty())
+	{
+		return;
+	}
+
+	if (!GraphData)
+	{
+		return;
+	}
+
+	GraphData->IncrementChangeBracket();
+
+	TArray<TObjectPtr<UEdGraphNode>> NodesToRemove = Nodes;
+	for (int32 NodeIndex = 0; NodeIndex < NodesToRemove.Num(); ++NodeIndex)
+	{
+		UScribbleEdGraphNode* ScribbleEdGraphNode = Cast<UScribbleEdGraphNode>(NodesToRemove[NodeIndex]);
+		if (!ScribbleEdGraphNode)
+		{
+			continue;
+		}
+		if (!SelectedNodes.Contains(ScribbleEdGraphNode->GetNodeId()))
+		{
+			continue;
+		}
+		RemoveScribbleNode(ScribbleEdGraphNode->GetNodeId());
+	}
+
+	SelectedNodes.Reset();
+	GraphData->DecrementChangeBracket();
+	
+	NotifyGraphChanged();
+}
+
+bool UScribbleEdGraph::IsNodeSelected(const FGuid& InNodeId) const
+{
+	return SelectedNodes.Contains(InNodeId);
+}
+
+bool UScribbleEdGraph::IsNodeSelected(const TSharedPtr<FScribbleNode>& InNode) const
+{
+	if (!InNode)
+	{
+		return false;
+	}
+	return IsNodeSelected(InNode->GetId());
+}
+
+bool UScribbleEdGraph::IsNodeSelected(const UScribbleEdGraphNode* InNode) const
+{
+	if (!InNode)
+	{
+		return false;
+	}
+	return IsNodeSelected(InNode->GetNodeId());
+}
+
+void UScribbleEdGraph::GroupSelectedNodes()
+{
+	if (!GraphData)
+	{
+		return;
+	}
+	
+	TArray<TSharedPtr<FScribbleNode>> NodesToGroup;
+	for (const FGuid& NodeId : SelectedNodes)
+	{
+		if (const UScribbleEdGraphNode* Node = FindNode(NodeId))
+		{
+			if (TSharedPtr<FScribbleNode> ScribbleNode = Node->GetScribbleNodePtr())
+			{
+				NodesToGroup.Add(ScribbleNode);
+			}
+		}
+	}
+	
+	TSharedPtr<FScribbleNode> GroupedNode = GraphData->GroupNodes(NodesToGroup);
+	if (!GroupedNode)
+	{
+		return;
+	}
+
+	for (const TSharedPtr<FScribbleNode>& ScribbleNode : NodesToGroup)
+	{
+		RemoveScribbleNode(ScribbleNode);
+	}
+
+	AddScribbleNode(GroupedNode);
+	if (IsAnchorEnabled.Get(false))
+	{
+		GroupedNode->SetAnchor(GetGraphData()->GetCurrentAnchor());
+	}
+	NotifyGraphChanged();
+
+	SelectNodes({GroupedNode->GetId()});
+}
+
+void UScribbleEdGraph::UngroupSelectedNodes()
+{
+	if (!GraphData)
+	{
+		return;
+	}
+
+	if (SelectedNodes.Num() != 1)
+	{
+		return;
+	}
+
+	TSharedPtr<FScribbleNode> GroupedNode;
+	if (const UScribbleEdGraphNode* Node = FindNode(SelectedNodes[0]))
+	{
+		if (TSharedPtr<FScribbleNode> ScribbleNode = Node->GetScribbleNodePtr())
+		{
+			GroupedNode = ScribbleNode;
+		}
+	}
+
+	TArray<TSharedPtr<FScribbleNode>> UngroupedNodes = GraphData->UngroupNode(GroupedNode);
+	if (UngroupedNodes.IsEmpty())
+	{
+		return;
+	}
+
+	RemoveScribbleNode(GroupedNode);
+
+	TArray<FGuid> NodesToSelect;
+	for (const TSharedPtr<FScribbleNode>& UngroupedNode : UngroupedNodes)
+	{
+		AddScribbleNode(UngroupedNode);
+		NodesToSelect.Add(UngroupedNode->GetId());
+	}
+	NotifyGraphChanged();
+
+	SelectNodes(NodesToSelect);
+}
+
+void UScribbleEdGraph::NotifyGraphChanged()
+{
+	if (GraphData)
+	{
+		if (GraphData->ChangeBracket > 0)
+		{
+			return;
+		}
+	}
+	Super::NotifyGraphChanged();
+}
+
+FBox2D UScribbleEdGraph::GetAllNodesBounds() const
+{
+	if (Nodes.IsEmpty())
+	{
+		return FBox2D();
+	}
+
+	FBox2D Bounds(EForceInit::ForceInit);
+	for (const TObjectPtr<UEdGraphNode>& Node : Nodes)
+	{
+		const FVector2D TopLeft(Node->NodePosX, Node->NodePosY);
+		Bounds += TopLeft;
+		Bounds += TopLeft + FVector2D(Node->GetSize());
+	}
+	return Bounds;
+}
+
+FBox2D UScribbleEdGraph::GetSelectedNodesBounds() const
+{
+	if (Nodes.IsEmpty() || SelectedNodes.IsEmpty())
+	{
+		return FBox2D();
+	}
+
+	return GetNodesBounds(SelectedNodes);
+}
+
+void UScribbleEdGraph::SelectNodes(const TArray<FGuid>& InNodeIds)
+{
+	SelectedNodes.Reset();
+	
+	TSet<const UEdGraphNode*> NodesToSelect;
+	for (const FGuid& NodeId : InNodeIds)
+	{
+		if (const UScribbleEdGraphNode* Node = FindNode(NodeId))
+		{
+			SelectedNodes.Add(NodeId);
+			NodesToSelect.Add(Node);
+		}
+	}
+
+	if (NodesToSelect.IsEmpty())
+	{
+		return;
+	}
+
+	if (bIsSelecting)
+	{
+		return;
+	}
+	const TGuardValue<bool> _(bIsSelecting, true);
+	SelectNodeSet(NodesToSelect);
+}
diff --git a/Engine/Plugins/Experimental/Animation/Scribble/Source/ScribbleEditor/Private/ScribbleEdGraphCommands.cpp b/Engine/Plugins/Experimental/Animation/Scribble/Source/ScribbleEditor/Private/ScribbleEdGraphCommands.cpp
new file mode 100644
--- /dev/null
+++ b/Engine/Plugins/Experimental/Animation/Scribble/Source/ScribbleEditor/Private/ScribbleEdGraphCommands.cpp
@@ -0,0 +1,18 @@
+// Copyright Epic Games, Inc. All Rights Reserved.
+
+#include "ScribbleEdGraphCommands.h"
+
+#define LOCTEXT_NAMESPACE "ScribbleEdGraphCommands"
+
+void FScribbleEdGraphCommands::RegisterCommands()
+{
+	UI_COMMAND(FrameSelection, "Frame Selection", "Expands and frames the selection graph", EUserInterfaceActionType::Button, FInputChord(EKeys::F));
+	UI_COMMAND(DeleteSelection, "Delete Selection", "Delete the selected elements from the graph", EUserInterfaceActionType::Button, FInputChord(EKeys::Delete));
+	UI_COMMAND(SelectAll, "Select All", "Select all scribble elements", EUserInterfaceActionType::Button, FInputChord(EModifierKey::Control, EKeys::A));
+	UI_COMMAND(GroupSelection, "Group Selection", "Groups the selected elements in the graph", EUserInterfaceActionType::Button, FInputChord(EModifierKey::Control, EKeys::G));
+	UI_COMMAND(UngroupSelection, "Ungroup Selection", "Ungroups the selected elements in the graph", EUserInterfaceActionType::Button, FInputChord(EModifierKey::Control | EModifierKey::Shift, EKeys::G));
+	UI_COMMAND(SelectionTool, "Selection Tool", "Enables / Disables the Selection Tool", EUserInterfaceActionType::Button, FInputChord(EKeys::V));
+	UI_COMMAND(BrushTool, "Brush Tool", "Enables / Disables the Brush Tool", EUserInterfaceActionType::Button, FInputChord(EKeys::B));
+}
+
+#undef LOCTEXT_NAMESPACE
diff --git a/Engine/Plugins/Experimental/Animation/Scribble/Source/ScribbleEditor/Private/ScribbleEdGraphNode.cpp b/Engine/Plugins/Experimental/Animation/Scribble/Source/ScribbleEditor/Private/ScribbleEdGraphNode.cpp
new file mode 100644
--- /dev/null
+++ b/Engine/Plugins/Experimental/Animation/Scribble/Source/ScribbleEditor/Private/ScribbleEdGraphNode.cpp
@@ -0,0 +1,122 @@
+// Copyright Epic Games, Inc. All Rights Reserved.
+
+#include "ScribbleEdGraphNode.h"
+#include "ScribbleEdGraph.h"
+#include "Templates/Casts.h"
+
+#include UE_INLINE_GENERATED_CPP_BY_NAME(ScribbleEdGraphNode)
+
+class UObject;
+
+UScribbleEdGraphNode::UScribbleEdGraphNode(const FObjectInitializer& ObjectInitializer)
+	: Super(ObjectInitializer)
+{
+}
+
+UScribbleEdGraphNode::~UScribbleEdGraphNode()
+{
+	FScribbleNode* ScribbleNode = GetScribbleNode();
+	if(!ScribbleNode)
+	{
+		return;
+	}
+	ScribbleNode->OnPositionChanged().Remove(PositionChangedDelegateHandle);
+	ScribbleNode->OnSizeChanged().Remove(SizeChangedDelegateHandle);
+}
+
+UScribbleEdGraph* UScribbleEdGraphNode::GetScribbleEdGraph() const
+{
+	return CastChecked<UScribbleEdGraph>(GetOuter());
+}
+
+void UScribbleEdGraphNode::Initialize(const TSharedPtr<FScribbleNode>& InScribbleNode)
+{
+	if (!InScribbleNode)
+	{
+		return;
+	}
+	WeakScribbleNode = InScribbleNode;
+	
+	NodePosX = InScribbleNode->GetPosition().X;
+	NodePosY = InScribbleNode->GetPosition().Y;
+	NodeWidth = InScribbleNode->GetSize().X;
+	NodeHeight = InScribbleNode->GetSize().Y;
+
+	TWeakObjectPtr<UScribbleEdGraphNode> WeakThis = this;
+	PositionChangedDelegateHandle = InScribbleNode->OnPositionChanged().AddLambda([WeakThis](const FVector2f& InNewPosition)
+	{
+		if(!WeakThis.IsValid())
+		{
+			return;
+		}
+		UScribbleEdGraphNode* StrongThis = WeakThis.Get();
+		StrongThis->NodePosX = InNewPosition.X;
+		StrongThis->NodePosY = InNewPosition.Y;
+	});
+
+	SizeChangedDelegateHandle = InScribbleNode->OnSizeChanged().AddLambda([WeakThis](const FVector2f& InNewSize)
+	{
+		if(!WeakThis.IsValid())
+		{
+			return;
+		}
+		UScribbleEdGraphNode* StrongThis = WeakThis.Get();
+		StrongThis->NodeWidth = InNewSize.X;
+		StrongThis->NodeHeight = InNewSize.Y;
+	});
+}
+
+const FScribbleNode* UScribbleEdGraphNode::GetScribbleNode() const
+{
+	if (WeakScribbleNode.IsValid())
+	{
+		return WeakScribbleNode.Pin().Get();
+	}
+	return nullptr;
+}
+
+FScribbleNode* UScribbleEdGraphNode::GetScribbleNode()
+{
+	if (WeakScribbleNode.IsValid())
+	{
+		return WeakScribbleNode.Pin().Get();
+	}
+	return nullptr;
+}
+
+TSharedPtr<FScribbleNode> UScribbleEdGraphNode::GetScribbleNodePtr() const
+{
+	if (WeakScribbleNode.IsValid())
+	{
+		return WeakScribbleNode.Pin();
+	}
+	return nullptr;
+}
+
+const FScribbleGraphData* UScribbleEdGraphNode::GetScribbleGraphData() const
+{
+	if (const FScribbleNode* ScribbleNode = GetScribbleNode())
+	{
+		return ScribbleNode->GetGraph(); 
+	}
+	return nullptr;
+}
+
+FScribbleGraphData* UScribbleEdGraphNode::GetScribbleGraphData()
+{
+	if (FScribbleNode* ScribbleNode = GetScribbleNode())
+	{
+		return ScribbleNode->GetGraph(); 
+	}
+	return nullptr;
+}
+
+const FGuid& UScribbleEdGraphNode::GetNodeId() const
+{
+	if (const FScribbleNode* ScribbleNode = GetScribbleNode())
+	{
+		return ScribbleNode->GetId();
+	}
+	static constexpr FGuid InvalidGuid = FGuid();
+	return InvalidGuid;
+}
diff --git a/Engine/Plugins/Experimental/Animation/Scribble/Source/ScribbleEditor/Private/ScribbleEdGraphPanelNodeFactory.cpp b/Engine/Plugins/Experimental/Animation/Scribble/Source/ScribbleEditor/Private/ScribbleEdGraphPanelNodeFactory.cpp
new file mode 100644
--- /dev/null
+++ b/Engine/Plugins/Experimental/Animation/Scribble/Source/ScribbleEditor/Private/ScribbleEdGraphPanelNodeFactory.cpp
@@ -0,0 +1,21 @@
+// Copyright Epic Games, Inc. All Rights Reserved.
+
+#include "ScribbleEdGraphPanelNodeFactory.h"
+
+#include "EdGraph/EdGraphNode.h"
+#include "ScribbleEdGraphNode.h"
+#include "SScribbleGraphNode.h"
+#include "SGraphNode.h"
+#include "Widgets/DeclarativeSyntaxSupport.h"
+
+TSharedPtr<SGraphNode> FScribbleEdGraphPanelNodeFactory::CreateNode(UEdGraphNode* Node) const
+{
+	if (UScribbleEdGraphNode* ScribbleEdGraphNode = Cast<UScribbleEdGraphNode>(Node))
+	{
+		TSharedRef<SGraphNode> GraphNode = SNew(SScribbleGraphNode, ScribbleEdGraphNode);
+		//GraphNode->SlatePrepass();
+		//ScribbleEdGraphNode->SetDimensions(GraphNode->GetDesiredSize());
+		return GraphNode;
+	}
+	return nullptr;
+}
diff --git a/Engine/Plugins/Experimental/Animation/Scribble/Source/ScribbleEditor/Private/ScribbleEdGraphSchema.cpp b/Engine/Plugins/Experimental/Animation/Scribble/Source/ScribbleEditor/Private/ScribbleEdGraphSchema.cpp
new file mode 100644
--- /dev/null
+++ b/Engine/Plugins/Experimental/Animation/Scribble/Source/ScribbleEditor/Private/ScribbleEdGraphSchema.cpp
@@ -0,0 +1,78 @@
+// Copyright Epic Games, Inc. All Rights Reserved.
+
+#include "ScribbleEdGraphSchema.h"
+
+#include "ScribbleEdGraphNode.h"
+#include "ScribbleEdGraph.h"
+
+#include UE_INLINE_GENERATED_CPP_BY_NAME(ScribbleEdGraphSchema)
+
+#define LOCTEXT_NAMESPACE "ScribbleEdGraphSchema"
+
+UScribbleEdGraphSchema::UScribbleEdGraphSchema(const FObjectInitializer& ObjectInitializer)
+	: Super(ObjectInitializer)
+{
+}
+
+UScribbleEdGraphNode* UScribbleEdGraphSchema::CreateGraphNode(UScribbleEdGraph* InGraph, const FGuid& InNodeId) const
+{
+	if (!InGraph)
+	{
+		return nullptr;
+	}
+
+	const FScribbleGraphData* GraphData = InGraph->GetGraphData();
+	if (!GraphData)
+	{
+		return nullptr;
+	}
+	
+	static constexpr bool bSelectNewNode = false;
+	
+	if (TSharedPtr<FScribbleNode> ScribbleNode = GraphData->FindNodePtr(InNodeId))
+	{
+		FGraphNodeCreator<UScribbleEdGraphNode> GraphNodeCreator(*InGraph);
+		UScribbleEdGraphNode* ScribbleEdGraphNode = GraphNodeCreator.CreateNode(bSelectNewNode);
+		ScribbleEdGraphNode->Initialize(ScribbleNode);
+		GraphNodeCreator.Finalize();
+
+		InGraph->NodeIdLookup.Add(InNodeId, ScribbleEdGraphNode);
+
+		return ScribbleEdGraphNode;
+	}
+
+	return nullptr;
+}
+
+const FPinConnectionResponse UScribbleEdGraphSchema::CanCreateConnection(const UEdGraphPin* A, const UEdGraphPin* B) const
+{
+	FPinConnectionResponse Response;
+	Response.Response = CONNECT_RESPONSE_DISALLOW;
+	return Response;
+}
+
+void UScribbleEdGraphSchema::SetNodePosition(UEdGraphNode* InNode, const FVector2f& InPosition) const
+{
+	Super::SetNodePosition(InNode, InPosition);
+
+	if (UScribbleEdGraphNode* ScribbleEdGraphNode = Cast<UScribbleEdGraphNode>(InNode))
+	{
+		if (FScribbleNode* ScribbleNode = ScribbleEdGraphNode->GetScribbleNode())
+		{
+			ScribbleNode->SetPosition(InPosition);
+		}
+	}
+}
+
+void UScribbleEdGraphSchema::SetNodeSize(UEdGraphNode* InNode, const FVector2f& InSize) const
+{
+	if (UScribbleEdGraphNode* ScribbleEdGraphNode = Cast<UScribbleEdGraphNode>(InNode))
+	{
+		if (FScribbleNode* ScribbleNode = ScribbleEdGraphNode->GetScribbleNode())
+		{
+			ScribbleNode->SetSize(InSize);
+		}
+	}
+}
+
+#undef LOCTEXT_NAMESPACE
diff --git a/Engine/Plugins/Experimental/Animation/Scribble/Source/ScribbleEditor/Private/ScribbleEditorModule.cpp b/Engine/Plugins/Experimental/Animation/Scribble/Source/ScribbleEditor/Private/ScribbleEditorModule.cpp
new file mode 100644
--- /dev/null
+++ b/Engine/Plugins/Experimental/Animation/Scribble/Source/ScribbleEditor/Private/ScribbleEditorModule.cpp
@@ -0,0 +1,28 @@
+// Copyright Epic Games, Inc. All Rights Reserved.
+
+#include "Modules/ModuleManager.h"
+#include "ScribbleEdGraphCommands.h"
+#include "ScribbleEdGraphPanelNodeFactory.h"
+
+class FScribbleEditorModule : public IModuleInterface
+{
+public:
+	/** IModuleInterface implementation */
+	virtual void StartupModule() override
+	{
+		ScribbleEdGraphPanelNodeFactory = MakeShareable(new FScribbleEdGraphPanelNodeFactory());
+		FEdGraphUtilities::RegisterVisualNodeFactory(ScribbleEdGraphPanelNodeFactory);
+		FScribbleEdGraphCommands::Register();
+	}
+	
+	virtual void ShutdownModule() override
+	{
+		FEdGraphUtilities::UnregisterVisualNodeFactory(ScribbleEdGraphPanelNodeFactory);
+		FScribbleEdGraphCommands::Unregister();
+	}
+
+private:
+	TSharedPtr<FScribbleEdGraphPanelNodeFactory> ScribbleEdGraphPanelNodeFactory;
+};
+
+IMPLEMENT_MODULE(FScribbleEditorModule, ScribbleEditor)
diff --git a/Engine/Plugins/Experimental/Animation/Scribble/Source/ScribbleEditor/Public/SScribbleGraph.h b/Engine/Plugins/Experimental/Animation/Scribble/Source/ScribbleEditor/Public/SScribbleGraph.h
new file mode 100644
--- /dev/null
+++ b/Engine/Plugins/Experimental/Animation/Scribble/Source/ScribbleEditor/Public/SScribbleGraph.h
@@ -0,0 +1,145 @@
+// Copyright Epic Games, Inc. All Rights Reserved.
+
+#pragma once
+
+#include "SGraphPanel.h"
+#include "EditorUndoClient.h"
+#include "ScribbleEdGraphNode.h"
+#include "ScopedTransaction.h"
+#include "Brushes/SlateRoundedBoxBrush.h"
+
+#define UE_API SCRIBBLEEDITOR_API
+
+DECLARE_DELEGATE_OneParam(FScribbleSetNodeType, EScribbleNodeType::Type);
+
+class SScribbleGraphPanel : public SGraphPanel, public FEditorUndoClient
+{
+public:
+
+	SLATE_BEGIN_ARGS(SScribbleGraphPanel)
+		: _ScribbleEnabled(true)
+		, _ScribbleNodeType(EScribbleNodeType::Invalid)
+		, _ScribbleColor()
+		, _ScribbleThickness()
+		, _ScribblePrecision()
+		, _ShouldDrawBackground(false)
+		, _ShouldDrawSurroundingShadow(false)
+		, _AllowNavigation(true)
+	{}
+	SLATE_ARGUMENT(TSharedPtr<FScribbleGraphData>, GraphData)
+	SLATE_ATTRIBUTE(bool, ScribbleEnabled)
+	SLATE_ATTRIBUTE(EScribbleNodeType::Type, ScribbleNodeType)
+	SLATE_EVENT(FScribbleSetNodeType, OnSetNodeType)
+	SLATE_ATTRIBUTE(FLinearColor, ScribbleColor)
+	SLATE_ATTRIBUTE(float, ScribbleThickness)
+	SLATE_ATTRIBUTE(float, ScribblePrecision)
+	SLATE_ATTRIBUTE(bool, ShouldDrawBackground)
+	SLATE_ATTRIBUTE(bool, ShouldDrawSurroundingShadow)
+	SLATE_ARGUMENT(bool, AllowNavigation)
+	SLATE_ARGUMENT(TWeakPtr<SGraphPanel>, PanelToSync)
+	SLATE_END_ARGS()
+
+	SScribbleGraphPanel();
+
+	UE_API virtual ~SScribbleGraphPanel() override;
+
+	/** Constructs this widget with InArgs */
+	UE_API void Construct(const FArguments& InArgs);
+
+	UE_API bool IsScribbleEnabled() const;
+	UE_API bool HasActiveNodeType() const;
+	UE_API void ActivateSelectionMode();
+	UE_API EScribbleNodeType::Type GetActiveNodeType() const;
+	UE_API void SetActiveNodeType(EScribbleNodeType::Type InNodeType);
+
+	UE_API UScribbleEdGraph* GetScribbleEdGraph();
+	UE_API const UScribbleEdGraph* GetScribbleEdGraph() const;
+	UE_API FScribbleGraphData* GetScribbleGraph();
+	UE_API const FScribbleGraphData* GetScribbleGraph() const;
+
+	UE_API virtual FReply OnMouseButtonDown( const FGeometry& MyGeometry, const FPointerEvent& MouseEvent ) override;
+	UE_API virtual void OnMouseCaptureLost(const FCaptureLostEvent& CaptureLostEvent) override;
+	UE_API virtual FReply OnMouseButtonUp( const FGeometry& MyGeometry, const FPointerEvent& MouseEvent ) override;
+	UE_API virtual FReply OnMouseMove( const FGeometry& MyGeometry, const FPointerEvent& MouseEvent ) override;
+	UE_API virtual void OnMouseLeave( const FPointerEvent& MouseEvent ) override;
+	UE_API virtual FReply OnKeyDown(const FGeometry& MyGeometry, const FKeyEvent& InKeyEvent) override;
+	UE_API virtual void Tick(const FGeometry& AllottedGeometry, const double InCurrentTime, const float InDeltaTime) override;
+	UE_API virtual int32 OnPaint(const FPaintArgs& Args, const FGeometry& AllottedGeometry, const FSlateRect& MyCullingRect, FSlateWindowElementList& OutDrawElements, int32 LayerId, const FWidgetStyle& InWidgetStyle, bool bParentEnabled) const override;
+
+	/** FEditorUndoClient interface */
+	UE_API virtual void PostUndo(bool bSuccess) override;
+	UE_API virtual void PostRedo(bool bSuccess) override;
+
+private:
+
+	void BindCommands();
+	void UpdateScribbleVisibility();
+	void HandleSelectionChanged(const FGraphPanelSelectionSet& SelectionSet);
+
+	bool UpdateNewNode( const FGeometry& MyGeometry, const FPointerEvent& MouseEvent );
+	bool CancelNewNode();
+	bool CommitNewNode();
+	
+	UScribbleEdGraph* ScribbleGraphObj;
+	TSharedPtr<FUICommandList> CommandList;
+
+	TAttribute<bool> ScribbleEnabled;
+	TAttribute<EScribbleNodeType::Type> ScribbleNodeType;
+	TAttribute<float> ScribbleThickness;
+	TAttribute<float> ScribblePrecision;
+	TAttribute<FLinearColor> ScribbleColor;
+	FScribbleSetNodeType OnSetNodeType;
+
+	TOptional<bool> bIsScribbleDragging;
+	TSharedPtr<FScribbleNode> NewNode;
+	TSharedPtr<FScopedTransaction> Transaction;
+	bool bAllowNavigation;
+	TWeakPtr<SGraphPanel> WeakPanelToSync;
+
+	const TUniquePtr<FSlateRoundedBoxBrush> RoundedFrameBrush;
+};
+
+class SScribbleGraph : public SCompoundWidget
+{
+public:
+
+	SLATE_BEGIN_ARGS(SScribbleGraph)
+		: _ShouldDrawBackground(true)
+		, _ShouldDrawSurroundingShadow(false)
+		, _AllowNavigation(true)
+		, _ToolbarHAlignment(HAlign_Left)
+		, _ToolbarVAlignment(VAlign_Top)
+		, _ToolbarPadding(20)
+	{}
+	SLATE_ARGUMENT(TSharedPtr<FScribbleGraphData>, GraphData)
+	SLATE_ATTRIBUTE(bool, ShouldDrawBackground)
+	SLATE_ATTRIBUTE(bool, ShouldDrawSurroundingShadow)
+	SLATE_ARGUMENT(bool, AllowNavigation)
+	SLATE_ARGUMENT(EHorizontalAlignment, ToolbarHAlignment)
+	SLATE_ARGUMENT(EVerticalAlignment, ToolbarVAlignment)
+	SLATE_ARGUMENT(FMargin, ToolbarPadding)
+	SLATE_ARGUMENT(TWeakPtr<SGraphPanel>, PanelToSync)
+	SLATE_END_ARGS()
+
+	SScribbleGraph()
+		: bIsScribbleEnabled(true)
+	{
+	}
+
+	UE_API virtual ~SScribbleGraph() override;
+
+	/** Constructs this widget with InArgs */
+	UE_API void Construct(const FArguments& InArgs);
+
+protected:
+
+	FLinearColor GetScribbleColor() const;
+	float GetScribbleThickness() const;
+	float GetScribblePrecision() const;
+	bool IsScribbleEnabled() const { return bIsScribbleEnabled; }
+
+	TSharedPtr<SScribbleGraphPanel> GraphPanel;
+	bool bIsScribbleEnabled;
+};
+
+#undef UE_API
diff --git a/Engine/Plugins/Experimental/Animation/Scribble/Source/ScribbleEditor/Public/SScribbleGraphNode.h b/Engine/Plugins/Experimental/Animation/Scribble/Source/ScribbleEditor/Public/SScribbleGraphNode.h
new file mode 100644
--- /dev/null
+++ b/Engine/Plugins/Experimental/Animation/Scribble/Source/ScribbleEditor/Public/SScribbleGraphNode.h
@@ -0,0 +1,63 @@
+// Copyright Epic Games, Inc. All Rights Reserved.
+
+#pragma once
+
+#include "ScribbleGraph.h"
+#include "Internationalization/Text.h"
+#include "SGraphNode.h"
+#include "Styling/SlateColor.h"
+#include "Templates/SharedPointer.h"
+#include "Widgets/DeclarativeSyntaxSupport.h"
+
+#define UE_API SCRIBBLEEDITOR_API
+
+class SGraphPin;
+class SVerticalBox;
+class SWidget;
+struct FSlateBrush;
+
+class SScribbleGraphNode : public SGraphNode
+{
+public:
+	SLATE_BEGIN_ARGS(SScribbleGraphNode){}
+
+	SLATE_END_ARGS()
+
+	/** Constructs this widget with InArgs */
+	UE_API void Construct(const FArguments& InArgs, class UScribbleEdGraphNode* InNode);
+
+protected:
+	
+	// SGraphNode interface
+	virtual void CreatePinWidgets() override {}
+	virtual void AddPin(const TSharedRef<SGraphPin>& PinToAdd) override {}
+
+	virtual FText GetNodeTooltip() const override { return FText(); }
+	virtual FSlateColor GetNodeBodyColor() const override { return FLinearColor(0,0,0,0); }
+	virtual const FSlateBrush* GetNodeBodyBrush() const override { return nullptr; }
+	UE_API virtual const FSlateBrush* GetShadowBrush(bool bSelected) const override;
+	UE_API virtual void UpdateGraphNode() override;
+	UE_API virtual FReply OnMouseButtonDown(const FGeometry& MyGeometry, const FPointerEvent& MouseEvent) override;
+	UE_API virtual FReply OnMouseMove(const FGeometry& MyGeometry, const FPointerEvent& MouseEvent) override;
+	UE_API virtual FReply OnMouseButtonUp(const FGeometry& MyGeometry, const FPointerEvent& MouseEvent) override;
+	UE_API virtual void MoveTo(const FVector2f& NewPosition, FNodeSet& NodeFilter, bool bMarkDirty = true) override;
+	UE_API virtual int32 OnPaint(const FPaintArgs& Args, const FGeometry& AllottedGeometry, const FSlateRect& MyCullingRect, FSlateWindowElementList& OutDrawElements, int32 LayerId, const FWidgetStyle& InWidgetStyle, bool bParentEnabled) const override;
+	UE_API virtual FVector2D ComputeDesiredSize(float LayoutScaleMultiplier) const override;
+	virtual FVector2f GetDesiredSizeForMarquee2f() const override { return GetSizeFromScribbleNode(); }
+	UE_API virtual bool CanBeSelected(const FVector2f& MousePositionInNode) const override;
+
+protected:
+
+	const UScribbleEdGraphNode* GetScribbleEdGraphNode() const;
+	UScribbleEdGraphNode* GetScribbleEdGraphNode();
+	const FScribbleNode* GetScribbleNode() const;
+	FScribbleNode* GetScribbleNode();
+	const FScribbleGraphData* GetScribbleGraphData() const;
+	FScribbleGraphData* GetScribbleGraphData();
+
+	FVector2f GetSizeFromScribbleNode() const;
+
+	friend class SScribbleGraphPanel;
+};
+
+#undef UE_API
diff --git a/Engine/Plugins/Experimental/Animation/Scribble/Source/ScribbleEditor/Public/ScribbleEdGraph.h b/Engine/Plugins/Experimental/Animation/Scribble/Source/ScribbleEditor/Public/ScribbleEdGraph.h
new file mode 100644
--- /dev/null
+++ b/Engine/Plugins/Experimental/Animation/Scribble/Source/ScribbleEditor/Public/ScribbleEdGraph.h
@@ -0,0 +1,97 @@
+// Copyright Epic Games, Inc. All Rights Reserved.
+
+#pragma once
+
+#include "ScribbleEdGraphNode.h"
+#include "EdGraph/EdGraph.h"
+#include "ScribbleGraph.h"
+#include "ScribbleEdGraph.generated.h"
+
+class UScribbleEdGraphSchema;
+
+/**
+ * The UEdGraph implementation of a scribble graph.
+ * This is used for the ed graph side of things only and corresponds with
+ * the FScribbleGraphData as well as the SScribbleGraph.
+ */
+UCLASS()
+class UScribbleEdGraph : public UEdGraph
+{
+	GENERATED_UCLASS_BODY()
+
+public:
+	void Initialize(const TSharedPtr<FScribbleGraphData>& InScribbleGraphData);
+
+	FScribbleGraphData* GetGraphData();
+	const FScribbleGraphData* GetGraphData() const;
+
+	/** Rebuilds graph from selection */
+	void RebuildGraph();
+
+	/** Select the nodes */
+	void SelectNodes(const TArray<FGuid>& InNodeIds);
+
+	/** Get the scribble graph nodes we are displaying */
+	const TArray<TObjectPtr<UEdGraphNode>>& GetNodes() const { return Nodes; }
+
+	/** Returns the node given a Node Id */
+	const UScribbleEdGraphNode* FindNode(const FGuid& InNodeId) const;
+	UScribbleEdGraphNode* FindNode(const FGuid& InNodeId);
+
+	FGuid AddScribbleNode(const TSharedPtr<FScribbleNode>& InNode);
+	bool RemoveScribbleNode(const TSharedPtr<FScribbleNode>& InNode);
+	bool RemoveScribbleNode(const FGuid& InGuid);
+
+	/** Returns true if the node given a Node Id exists */
+	bool Contains(const FGuid& InNodeId) const;
+
+	void SelectAllNodes();
+	void ClearSelection();
+	void RemoveAllNodes();
+	TArray<FGuid> GetAllNodeIds() const;
+	const TArray<FGuid>& GetSelectedNodeIds() const { return SelectedNodes; }
+	void RemoveSelectedNodes();
+	bool IsNodeSelected(const FGuid& InNodeId) const;
+	bool IsNodeSelected(const TSharedPtr<FScribbleNode>& InNode) const;
+	bool IsNodeSelected(const UScribbleEdGraphNode* InNode) const;
+
+	void GroupSelectedNodes();
+	void UngroupSelectedNodes();
+
+	virtual void NotifyGraphChanged() override;
+
+private:
+	FBox2D GetAllNodesBounds() const;
+	FBox2D GetSelectedNodesBounds() const;
+
+	template<typename RangeType>
+	FBox2D GetNodesBounds(const RangeType& InNodeIds) const
+	{
+		FBox2D Bounds(EForceInit::ForceInit);
+		for (const FGuid& NodeId : InNodeIds)
+		{
+			if (const UScribbleEdGraphNode* Node = FindNode(NodeId))
+			{
+				const FVector2D TopLeft(Node->NodePosX, Node->NodePosY);
+				Bounds += TopLeft;
+				Bounds += TopLeft + FVector2D(Node->GetSize());
+			}
+		}
+		return Bounds;
+	}
+
+private:
+
+	TArray<FGuid> SelectedNodes;
+	bool bIsSelecting;
+
+	TMap<FGuid, UScribbleEdGraphNode*> NodeIdLookup;
+
+	TSharedPtr<FScribbleGraphData> GraphData;
+	TAttribute<bool> IsAnchorEnabled;
+
+	friend class SScribbleGraph;
+	friend class SScribbleGraphPanel;
+	friend class UScribbleEdGraphSchema;
+};
+
diff --git a/Engine/Plugins/Experimental/Animation/Scribble/Source/ScribbleEditor/Public/ScribbleEdGraphCommands.h b/Engine/Plugins/Experimental/Animation/Scribble/Source/ScribbleEditor/Public/ScribbleEdGraphCommands.h
new file mode 100644
--- /dev/null
+++ b/Engine/Plugins/Experimental/Animation/Scribble/Source/ScribbleEditor/Public/ScribbleEdGraphCommands.h
@@ -0,0 +1,46 @@
+// Copyright Epic Games, Inc. All Rights Reserved.
+
+#pragma once
+
+#include "CoreMinimal.h"
+#include "Framework/Commands/Commands.h"
+#include "Styling/AppStyle.h"
+
+class FScribbleEdGraphCommands : public TCommands<FScribbleEdGraphCommands>
+{
+public:
+	FScribbleEdGraphCommands() : TCommands<FScribbleEdGraphCommands>
+	(
+		"ScribbleEdGraph",
+		NSLOCTEXT("Contexts", "ScribbleEdGraph", "Scribble"),
+		NAME_None, // "MainFrame" // @todo Fix this crash
+		FAppStyle::Get().GetStyleSetName() // Icon Style Set
+	)
+	{}
+
+	/* frames the selection in the graph */
+	TSharedPtr<FUICommandInfo> FrameSelection;
+
+	/* delete the selected nodes in the graph */
+	TSharedPtr<FUICommandInfo> DeleteSelection;
+
+	/* select all scribble elements */
+	TSharedPtr<FUICommandInfo> SelectAll;
+
+	/* groups the selection in the graph */
+	TSharedPtr<FUICommandInfo> GroupSelection;
+
+	/* ungroups the selection in the graph */
+	TSharedPtr<FUICommandInfo> UngroupSelection;
+
+	/* enables / disables the select tool */
+	TSharedPtr<FUICommandInfo> SelectionTool;
+
+	/* enables / disables the brush tool */
+	TSharedPtr<FUICommandInfo> BrushTool;
+
+	/**
+	 * Initialize commands
+	 */
+	virtual void RegisterCommands() override;
+};
diff --git a/Engine/Plugins/Experimental/Animation/Scribble/Source/ScribbleEditor/Public/ScribbleEdGraphNode.h b/Engine/Plugins/Experimental/Animation/Scribble/Source/ScribbleEditor/Public/ScribbleEdGraphNode.h
new file mode 100644
--- /dev/null
+++ b/Engine/Plugins/Experimental/Animation/Scribble/Source/ScribbleEditor/Public/ScribbleEdGraphNode.h
@@ -0,0 +1,52 @@
+// Copyright Epic Games, Inc. All Rights Reserved.
+
+#pragma once
+
+#include "EdGraph/EdGraphNode.h"
+#include "ScribbleNode.h"
+#include "Textures/SlateIcon.h"
+#include "ScribbleEdGraphNode.generated.h"
+
+class UScribbleEdGraph;
+class UEdGraphPin;
+
+/**
+ * The UEdGraphNode implementation of a scribble graph node.
+ * This is used for the ed graph side of things only and corresponds with
+ * the FScribbleNode as well as the SScribbleGraphNode.
+ */
+UCLASS()
+class UScribbleEdGraphNode : public UEdGraphNode
+{
+	GENERATED_UCLASS_BODY()
+
+	virtual ~UScribbleEdGraphNode();
+
+	class UScribbleEdGraph* GetScribbleEdGraph() const;
+
+	void Initialize(const TSharedPtr<FScribbleNode>& InScribbleNode);
+
+	// UEdGraphNode implementation
+	virtual void AllocateDefaultPins() override {}
+	virtual bool ShowPaletteIconOnNode() const override { return false; }
+	virtual FSlateIcon GetIconAndTint(FLinearColor& OutColor) const override { return FSlateIcon(); }
+
+	const FScribbleNode* GetScribbleNode() const;
+	FScribbleNode* GetScribbleNode();
+	TSharedPtr<FScribbleNode> GetScribbleNodePtr() const;
+	const FScribbleGraphData* GetScribbleGraphData() const;
+	FScribbleGraphData* GetScribbleGraphData();
+
+	const FGuid& GetNodeId() const;
+
+protected:
+	
+	TWeakPtr<FScribbleNode> WeakScribbleNode;
+	FDelegateHandle PositionChangedDelegateHandle;
+	FDelegateHandle SizeChangedDelegateHandle;
+
+	friend class UScribbleEdGraph;
+	friend class UScribbleEdGraphSchema;
+};
+
+
diff --git a/Engine/Plugins/Experimental/Animation/Scribble/Source/ScribbleEditor/Public/ScribbleEdGraphPanelNodeFactory.h b/Engine/Plugins/Experimental/Animation/Scribble/Source/ScribbleEditor/Public/ScribbleEdGraphPanelNodeFactory.h
new file mode 100644
--- /dev/null
+++ b/Engine/Plugins/Experimental/Animation/Scribble/Source/ScribbleEditor/Public/ScribbleEdGraphPanelNodeFactory.h
@@ -0,0 +1,17 @@
+// Copyright Epic Games, Inc. All Rights Reserved.
+
+#pragma once
+
+#include "EdGraphUtilities.h"
+#include "Templates/SharedPointer.h"
+
+class UEdGraphNode;
+
+#define UE_API SCRIBBLEEDITOR_API
+
+class FScribbleEdGraphPanelNodeFactory : public FGraphPanelNodeFactory
+{
+	UE_API virtual TSharedPtr<class SGraphNode> CreateNode(UEdGraphNode* Node) const override;
+};
+
+#undef UE_API
diff --git a/Engine/Plugins/Experimental/Animation/Scribble/Source/ScribbleEditor/Public/ScribbleEdGraphSchema.h b/Engine/Plugins/Experimental/Animation/Scribble/Source/ScribbleEditor/Public/ScribbleEdGraphSchema.h
new file mode 100644
--- /dev/null
+++ b/Engine/Plugins/Experimental/Animation/Scribble/Source/ScribbleEditor/Public/ScribbleEdGraphSchema.h
@@ -0,0 +1,29 @@
+// Copyright Epic Games, Inc. All Rights Reserved.
+
+#pragma once
+
+#include "ScribbleEdGraph.h"
+#include "EdGraph/EdGraphSchema.h"
+#include "ScribbleEdGraphNode.h"
+#include "ScribbleEdGraphSchema.generated.h"
+
+class UScribbleEdGraph;
+class UScribbleEdGraphNode;
+
+UCLASS()
+class UScribbleEdGraphSchema : public UEdGraphSchema
+{
+	GENERATED_UCLASS_BODY()
+
+public:
+	// UEdGraphSchema interface
+	virtual FName GetParentContextMenuName() const override { return NAME_None; }
+	virtual bool ShouldAlwaysPurgeOnModification() const override { return false; }
+	virtual const FPinConnectionResponse CanCreateConnection(const UEdGraphPin* A, const UEdGraphPin* B) const override;
+	virtual void SetNodePosition(UEdGraphNode* InNode, const FVector2f& InPosition) const override;
+	virtual void SetNodeSize(UEdGraphNode* InNode, const FVector2f& InSize) const;
+	// End of UEdGraphSchema interface
+
+	virtual UScribbleEdGraphNode* CreateGraphNode(UScribbleEdGraph* InGraph, const FGuid& InNodeId) const;
+};
+
diff --git a/Engine/Plugins/Experimental/Animation/Scribble/Source/ScribbleEditor/ScribbleEditor.Build.cs b/Engine/Plugins/Experimental/Animation/Scribble/Source/ScribbleEditor/ScribbleEditor.Build.cs
new file mode 100644
--- /dev/null
+++ b/Engine/Plugins/Experimental/Animation/Scribble/Source/ScribbleEditor/ScribbleEditor.Build.cs
@@ -0,0 +1,38 @@
+// Copyright Epic Games, Inc. All Rights Reserved.
+
+namespace UnrealBuildTool.Rules
+{
+    public class ScribbleEditor : ModuleRules
+    {
+        public ScribbleEditor(ReadOnlyTargetRules Target) : base(Target)
+        {
+            PublicDependencyModuleNames.AddRange(
+				new string[]
+				{
+					"MainFrame",
+					"AppFramework",
+					"GraphEditor",
+					"Scribble"
+				}
+			);
+
+            PrivateDependencyModuleNames.AddRange(
+                new string[]
+                {
+                    "Core",
+                    "CoreUObject",
+					"Slate",
+                    "SlateCore",
+                    "InputCore",
+					"Engine",
+					"EditorFramework",
+					"UnrealEd",
+                    "EditorStyle",
+					"EditorWidgets",
+                    "ApplicationCore",
+                    "PropertyEditor",
+				}
+            );
+        }
+    }
+}
