Scribble: Fix uninitialized member

#rnx

diff --git a/Engine/Plugins/Experimental/Animation/Scribble/Source/Scribble/Public/ScribbleNode.h b/Engine/Plugins/Experimental/Animation/Scribble/Source/Scribble/Public/ScribbleNode.h
--- a/Engine/Plugins/Experimental/Animation/Scribble/Source/Scribble/Public/ScribbleNode.h
+++ b/Engine/Plugins/Experimental/Animation/Scribble/Source/Scribble/Public/ScribbleNode.h
@@ -67,7 +67,19 @@ struct FScribbleNode
 
 public:
 
-	FScribbleNode(const EScribbleNodeType::Type& InType = EScribbleNodeType::Invalid)
+	FScribbleNode()
+		: Type(EScribbleNodeType::Invalid)
+		, Guid(FGuid())
+		, Anchor(NAME_None)
+		, Position(FVector2f::ZeroVector)
+		, RelativePosition(FVector2f::ZeroVector)
+		, Size(1,1)
+		, ChangeBracket(0)
+		, AccumulatedChanges(0)
+	{
+	}
+
+	FScribbleNode(const EScribbleNodeType::Type& InType)
 		: Type(InType)
 		, Guid(FGuid::NewGuid())
 		, Anchor(NAME_None)
